First usable version of the render pages plugin.

diff --git a/Engine/Plugins/Experimental/RenderPages/RenderPages.uplugin b/Engine/Plugins/Experimental/RenderPages/RenderPages.uplugin
new file mode 100644
--- /dev/null
+++ b/Engine/Plugins/Experimental/RenderPages/RenderPages.uplugin
@@ -0,0 +1,52 @@
+{
+  "FileVersion": 3,
+  "Version": 1,
+  "VersionName": "1.0",
+  "FriendlyName": "Render Pages",
+  "Description": "Advanced pipeline for use in creating rendered cinematics.",
+  "Category": "Experimental",
+  "CreatedBy": "Epic Games, Inc.",
+  "CreatedByURL": "https://epicgames.com",
+  "DocsURL": "",
+  "MarketplaceURL": "",
+  "SupportURL": "",
+  "EnabledByDefault": false,
+  "CanContainContent": false,
+  "IsBetaVersion": false,
+  "IsExperimentalVersion": true,
+  "Installed": false,
+  "Modules": [
+    {
+      "Name": "RenderPages",
+      "Type": "Editor",
+      "LoadingPhase": "PostEngineInit"
+    },
+    {
+      "Name": "RenderPagesDeveloper",
+      "Type": "Editor",
+      "LoadingPhase": "PostEngineInit"
+    },
+    {
+      "Name": "RenderPagesEditor",
+      "Type": "Editor",
+      "LoadingPhase": "PostEngineInit"
+    }
+  ],
+  "Plugins": [
+    {
+      "Name": "RemoteControl",
+      "Enabled": true,
+      "Optional": false
+    },
+    {
+      "Name": "LevelSequenceEditor",
+      "Enabled": true,
+      "Optional": false
+    },
+    {
+      "Name": "MovieRenderPipeline",
+      "Enabled": true,
+      "Optional": false
+    }
+  ]
+}
\ No newline at end of file
diff --git a/Engine/Plugins/Experimental/RenderPages/Source/RenderPages/Private/Factories/RenderPagePropsSourceFactoryLocal.cpp b/Engine/Plugins/Experimental/RenderPages/Source/RenderPages/Private/Factories/RenderPagePropsSourceFactoryLocal.cpp
new file mode 100644
--- /dev/null
+++ b/Engine/Plugins/Experimental/RenderPages/Source/RenderPages/Private/Factories/RenderPagePropsSourceFactoryLocal.cpp
@@ -0,0 +1,14 @@
+// Copyright Epic Games, Inc. All Rights Reserved.
+
+#pragma once
+
+#include "Factories/RenderPagePropsSourceFactoryLocal.h"
+#include "RenderPage/RenderPagePropsSource.h"
+
+
+URenderPagePropsSourceBase* UE::RenderPages::Private::FRenderPagePropsSourceFactoryLocal::CreateInstance(UObject* Outer, UObject* PropsSourceOrigin)
+{
+	URenderPagePropsSourceLocal* PropsSource = NewObject<URenderPagePropsSourceLocal>(Outer);
+	PropsSource->SetSourceOrigin(PropsSourceOrigin);
+	return PropsSource;
+}
diff --git a/Engine/Plugins/Experimental/RenderPages/Source/RenderPages/Private/Factories/RenderPagePropsSourceFactoryLocal.h b/Engine/Plugins/Experimental/RenderPages/Source/RenderPages/Private/Factories/RenderPagePropsSourceFactoryLocal.h
new file mode 100644
--- /dev/null
+++ b/Engine/Plugins/Experimental/RenderPages/Source/RenderPages/Private/Factories/RenderPagePropsSourceFactoryLocal.h
@@ -0,0 +1,19 @@
+// Copyright Epic Games, Inc. All Rights Reserved.
+
+#pragma once
+
+#include "CoreMinimal.h"
+#include "Factories/IRenderPagePropsSourceFactory.h"
+
+
+namespace UE::RenderPages::Private
+{
+	/**
+	 * The factory class for URenderPagePropsSourceLocal.
+	 */
+	class RENDERPAGES_API FRenderPagePropsSourceFactoryLocal final : public IRenderPagePropsSourceFactory
+	{
+	public:
+		virtual URenderPagePropsSourceBase* CreateInstance(UObject* Outer, UObject* PropsSourceOrigin) override;
+	};
+}
diff --git a/Engine/Plugins/Experimental/RenderPages/Source/RenderPages/Private/Factories/RenderPagePropsSourceFactoryRemoteControl.cpp b/Engine/Plugins/Experimental/RenderPages/Source/RenderPages/Private/Factories/RenderPagePropsSourceFactoryRemoteControl.cpp
new file mode 100644
--- /dev/null
+++ b/Engine/Plugins/Experimental/RenderPages/Source/RenderPages/Private/Factories/RenderPagePropsSourceFactoryRemoteControl.cpp
@@ -0,0 +1,14 @@
+// Copyright Epic Games, Inc. All Rights Reserved.
+
+#pragma once
+
+#include "Factories/RenderPagePropsSourceFactoryRemoteControl.h"
+#include "RenderPage/RenderPagePropsSource.h"
+
+
+URenderPagePropsSourceBase* UE::RenderPages::Private::FRenderPagePropsSourceFactoryRemoteControl::CreateInstance(UObject* Outer, UObject* PropsSourceOrigin)
+{
+	URenderPagePropsSourceRemoteControl* PropsSource = NewObject<URenderPagePropsSourceRemoteControl>(Outer);
+	PropsSource->SetSourceOrigin(PropsSourceOrigin);
+	return PropsSource;
+}
diff --git a/Engine/Plugins/Experimental/RenderPages/Source/RenderPages/Private/Factories/RenderPagePropsSourceFactoryRemoteControl.h b/Engine/Plugins/Experimental/RenderPages/Source/RenderPages/Private/Factories/RenderPagePropsSourceFactoryRemoteControl.h
new file mode 100644
--- /dev/null
+++ b/Engine/Plugins/Experimental/RenderPages/Source/RenderPages/Private/Factories/RenderPagePropsSourceFactoryRemoteControl.h
@@ -0,0 +1,19 @@
+// Copyright Epic Games, Inc. All Rights Reserved.
+
+#pragma once
+
+#include "CoreMinimal.h"
+#include "Factories/IRenderPagePropsSourceFactory.h"
+
+
+namespace UE::RenderPages::Private
+{
+	/**
+	 * The factory class for URenderPagePropsSourceRemoteControl.
+	 */
+	class RENDERPAGES_API FRenderPagePropsSourceFactoryRemoteControl final : public IRenderPagePropsSourceFactory
+	{
+	public:
+		virtual URenderPagePropsSourceBase* CreateInstance(UObject* Outer, UObject* PropsSourceOrigin) override;
+	};
+}
diff --git a/Engine/Plugins/Experimental/RenderPages/Source/RenderPages/Private/RenderPage/RenderPageCollection.cpp b/Engine/Plugins/Experimental/RenderPages/Source/RenderPages/Private/RenderPage/RenderPageCollection.cpp
new file mode 100644
--- /dev/null
+++ b/Engine/Plugins/Experimental/RenderPages/Source/RenderPages/Private/RenderPage/RenderPageCollection.cpp
@@ -0,0 +1,697 @@
+// Copyright Epic Games, Inc. All Rights Reserved.
+
+#include "RenderPage/RenderPageCollection.h"
+#include "RenderPage/RenderPageManager.h"
+#include "RenderPagesUtils.h"
+#include "IRenderPagesModule.h"
+#include "LevelSequence.h"
+#include "MoviePipelineMasterConfig.h"
+#include "MoviePipelineOutputSetting.h"
+#include "MovieScene.h"
+#include "UObject/ObjectSaveContext.h"
+
+
+URenderPage::URenderPage()
+	: Id(FGuid::NewGuid())
+	, Sequence(nullptr)
+	, bOverrideStartFrame(false)
+	, CustomStartFrame(0)
+	, bOverrideEndFrame(false)
+	, CustomEndFrame(0)
+	, bOverrideResolution(false)
+	, CustomResolution(FIntPoint(3840, 2160))
+	, bIsEnabled(true)
+	, RenderPreset(nullptr)
+{}
+
+TOptional<int32> URenderPage::GetSequenceStartFrame() const
+{
+	if (!IsValid(Sequence))
+	{
+		return TOptional<int32>();
+	}
+
+	const UMovieScene* MovieScene = Sequence->MovieScene;
+	const FFrameRate TickResolution = MovieScene->GetTickResolution();
+	const FFrameRate DisplayRate = MovieScene->GetDisplayRate();
+	const UMoviePipelineOutputSetting* Settings = GetRenderPresetOutputSettings();
+
+	if (bOverrideStartFrame || (IsValid(Settings) && Settings->bUseCustomPlaybackRange))
+	{
+		int32 Frame = (bOverrideStartFrame ? CustomStartFrame : Settings->CustomStartFrame);
+		if (IsValid(Settings) && Settings->bUseCustomFrameRate)
+		{
+			return FMath::FloorToInt(Frame / (Settings->OutputFrameRate / DisplayRate).AsDecimal());
+		}
+		return Frame;
+	}
+
+	const int32 StartFrameNumber = MovieScene->GetPlaybackRange().GetLowerBoundValue().Value;
+	return FMath::FloorToInt(StartFrameNumber / (TickResolution / DisplayRate).AsDecimal());
+}
+
+TOptional<int32> URenderPage::GetSequenceEndFrame() const
+{
+	if (!IsValid(Sequence))
+	{
+		return TOptional<int32>();
+	}
+
+	const UMovieScene* MovieScene = Sequence->MovieScene;
+	const FFrameRate TickResolution = MovieScene->GetTickResolution();
+	const FFrameRate DisplayRate = MovieScene->GetDisplayRate();
+	const UMoviePipelineOutputSetting* Settings = GetRenderPresetOutputSettings();
+
+	if (bOverrideEndFrame || (IsValid(Settings) && Settings->bUseCustomPlaybackRange))
+	{
+		int32 Frame = (bOverrideEndFrame ? CustomEndFrame : Settings->CustomEndFrame);
+		if (IsValid(Settings) && Settings->bUseCustomFrameRate)
+		{
+			return FMath::FloorToInt(Frame / (Settings->OutputFrameRate / DisplayRate).AsDecimal());
+		}
+		return Frame;
+	}
+
+	const int32 EndFrameNumber = MovieScene->GetPlaybackRange().GetUpperBoundValue().Value;
+	return FMath::FloorToInt(EndFrameNumber / (TickResolution / DisplayRate).AsDecimal());
+}
+
+bool URenderPage::SetSequenceStartFrame(const int32 NewCustomStartFrame)
+{
+	int32 StartFrame = GetStartFrame().Get(0);
+	SetIsCustomStartFrame(true);
+	SetCustomStartFrame(StartFrame);
+
+	TOptional<int32> SequenceStartFrame = GetSequenceStartFrame();
+	while (SequenceStartFrame.IsSet() && (SequenceStartFrame.Get(0) > NewCustomStartFrame))
+	{
+		StartFrame--;
+		if (StartFrame <= -2147483648)
+		{
+			return false;
+		}
+		SetCustomStartFrame(StartFrame);
+		SequenceStartFrame = GetSequenceStartFrame();
+	}
+	while (SequenceStartFrame.IsSet() && (SequenceStartFrame.Get(0) < NewCustomStartFrame))
+	{
+		StartFrame++;
+		if (StartFrame >= 2147483647)
+		{
+			return false;
+		}
+		SetCustomStartFrame(StartFrame);
+		SequenceStartFrame = GetSequenceStartFrame();
+	}
+	return (SequenceStartFrame.IsSet() && (SequenceStartFrame.Get(0) == NewCustomStartFrame));
+}
+
+bool URenderPage::SetSequenceEndFrame(const int32 NewCustomStartFrame)
+{
+	int32 EndFrame = GetEndFrame().Get(0);
+	SetIsCustomEndFrame(true);
+	SetCustomEndFrame(EndFrame);
+
+	TOptional<int32> SequenceEndFrame = GetSequenceEndFrame();
+	while (SequenceEndFrame.IsSet() && (SequenceEndFrame.Get(0) > NewCustomStartFrame))
+	{
+		EndFrame--;
+		if (EndFrame <= -2147483648)
+		{
+			return false;
+		}
+		SetCustomEndFrame(EndFrame);
+		SequenceEndFrame = GetSequenceEndFrame();
+	}
+	while (SequenceEndFrame.IsSet() && (SequenceEndFrame.Get(0) < NewCustomStartFrame))
+	{
+		EndFrame++;
+		if (EndFrame >= 2147483647)
+		{
+			return false;
+		}
+		SetCustomEndFrame(EndFrame);
+		SequenceEndFrame = GetSequenceEndFrame();
+	}
+	return (SequenceEndFrame.IsSet() && (SequenceEndFrame.Get(0) == NewCustomStartFrame));
+}
+
+TOptional<int32> URenderPage::GetStartFrame() const
+{
+	if (bOverrideStartFrame)
+	{
+		return CustomStartFrame;
+	}
+
+	const UMoviePipelineOutputSetting* Settings = GetRenderPresetOutputSettings();
+	if (IsValid(Settings) && Settings->bUseCustomPlaybackRange)
+	{
+		return Settings->CustomStartFrame;
+	}
+
+	if (!IsValid(Sequence))
+	{
+		return TOptional<int32>();
+	}
+
+	const UMovieScene* MovieScene = Sequence->MovieScene;
+	const FFrameRate TickResolution = MovieScene->GetTickResolution();
+	FFrameRate DisplayRate = MovieScene->GetDisplayRate();
+	if (IsValid(Settings) && Settings->bUseCustomFrameRate)
+	{
+		DisplayRate = Settings->OutputFrameRate;
+	}
+	const int32 StartFrameNumber = MovieScene->GetPlaybackRange().GetLowerBoundValue().Value;
+	return FMath::FloorToInt(StartFrameNumber / (TickResolution / DisplayRate).AsDecimal());
+}
+
+TOptional<int32> URenderPage::GetEndFrame() const
+{
+	if (bOverrideEndFrame)
+	{
+		return CustomEndFrame;
+	}
+
+	const UMoviePipelineOutputSetting* Settings = GetRenderPresetOutputSettings();
+	if (IsValid(Settings) && Settings->bUseCustomPlaybackRange)
+	{
+		return Settings->CustomEndFrame;
+	}
+
+	if (!IsValid(Sequence))
+	{
+		return TOptional<int32>();
+	}
+
+	const UMovieScene* MovieScene = Sequence->MovieScene;
+	const FFrameRate TickResolution = MovieScene->GetTickResolution();
+	FFrameRate DisplayRate = MovieScene->GetDisplayRate();
+	if (IsValid(Settings) && Settings->bUseCustomFrameRate)
+	{
+		DisplayRate = Settings->OutputFrameRate;
+	}
+	const int32 EndFrameNumber = MovieScene->GetPlaybackRange().GetUpperBoundValue().Value;
+	return FMath::FloorToInt(EndFrameNumber / (TickResolution / DisplayRate).AsDecimal());
+}
+
+TOptional<double> URenderPage::GetStartTime() const
+{
+	if (!IsValid(Sequence))
+	{
+		return TOptional<double>();
+	}
+
+	const TOptional<int32> StartFrame = GetStartFrame();
+	if (!StartFrame.IsSet())
+	{
+		return TOptional<double>();
+	}
+
+	const UMoviePipelineOutputSetting* Settings = GetRenderPresetOutputSettings();
+
+	const UMovieScene* MovieScene = Sequence->MovieScene;
+	FFrameRate DisplayRate = MovieScene->GetDisplayRate();
+	if (IsValid(Settings) && Settings->bUseCustomFrameRate)
+	{
+		DisplayRate = Settings->OutputFrameRate;
+	}
+	return StartFrame.Get(0) / DisplayRate.AsDecimal();
+}
+
+TOptional<double> URenderPage::GetEndTime() const
+{
+	if (!IsValid(Sequence))
+	{
+		return TOptional<double>();
+	}
+
+	const TOptional<int32> EndFrame = GetEndFrame();
+	if (!EndFrame.IsSet())
+	{
+		return TOptional<double>();
+	}
+
+	const UMoviePipelineOutputSetting* Settings = GetRenderPresetOutputSettings();
+
+	const UMovieScene* MovieScene = Sequence->MovieScene;
+	FFrameRate DisplayRate = MovieScene->GetDisplayRate();
+	if (IsValid(Settings) && Settings->bUseCustomFrameRate)
+	{
+		DisplayRate = Settings->OutputFrameRate;
+	}
+	return EndFrame.Get(0) / DisplayRate.AsDecimal();
+}
+
+TOptional<double> URenderPage::GetDurationInSeconds() const
+{
+	if (!IsValid(Sequence))
+	{
+		return TOptional<double>();
+	}
+
+	const TOptional<int32> StartFrame = GetStartFrame();
+	const TOptional<int32> EndFrame = GetEndFrame();
+	if (!StartFrame.IsSet() || !EndFrame.IsSet() || (StartFrame.Get(0) > EndFrame.Get(0)))
+	{
+		return TOptional<double>();
+	}
+
+	const UMoviePipelineOutputSetting* Settings = GetRenderPresetOutputSettings();
+
+	const UMovieScene* MovieScene = Sequence->MovieScene;
+	FFrameRate DisplayRate = MovieScene->GetDisplayRate();
+	if (IsValid(Settings) && Settings->bUseCustomFrameRate)
+	{
+		DisplayRate = Settings->OutputFrameRate;
+	}
+	return (EndFrame.Get(0) - StartFrame.Get(0)) / DisplayRate.AsDecimal();
+}
+
+double URenderPage::GetOutputAspectRatio() const
+{
+	const UMoviePipelineOutputSetting* Settings = GetRenderPresetOutputSettings();
+	if (IsValid(Settings))
+	{
+		return static_cast<double>(Settings->OutputResolution.X) / static_cast<double>(Settings->OutputResolution.Y);
+	}
+
+	const UMoviePipelineOutputSetting* DefaultSettings = GetDefault<UMoviePipelineOutputSetting>();
+	return static_cast<double>(DefaultSettings->OutputResolution.X) / static_cast<double>(DefaultSettings->OutputResolution.Y);
+}
+
+bool URenderPage::MatchesSearchTerm(const FString& SearchTerm) const
+{
+	if (SearchTerm.TrimStartAndEnd().Len() <= 0)
+	{
+		return true;
+	}
+	TArray<FString> Parts;
+	SearchTerm.ParseIntoArray(Parts, TEXT(" "), true);
+	for (const FString& Part : Parts)
+	{
+		FString TrimmedPart = Part.TrimStartAndEnd();
+		if (TrimmedPart.Len() <= 0)
+		{
+			continue;
+		}
+		if ((PageId.Find(TrimmedPart) != INDEX_NONE) || (PageName.Find(TrimmedPart) != INDEX_NONE) || (OutputDirectory.Find(TrimmedPart) != INDEX_NONE))
+		{
+			continue;
+		}
+		if (IsValid(RenderPreset) && (RenderPreset.GetPath().Find(TrimmedPart) != INDEX_NONE))
+		{
+			continue;
+		}
+		return false;
+	}
+	return true;
+}
+
+FString URenderPage::PurgePageIdOrReturnEmptyString(const FString& NewPageId)
+{
+	static FString ValidCharacters = TEXT("abcdefghijklmnopqrstuvwxyzABCDEFGHIJKLMNOPQRSTUVWXYZ0123456789-_");
+
+	FString Result;
+	for (TCHAR NewPageIdChar : NewPageId)
+	{
+		int32 Index;
+		if (ValidCharacters.FindChar(NewPageIdChar, Index))
+		{
+			Result += NewPageIdChar;
+		}
+	}
+	return Result;
+}
+
+FString URenderPage::PurgePageId(const FString& NewPageId)
+{
+	FString Result = PurgePageIdOrReturnEmptyString(NewPageId);
+	if (Result.IsEmpty())
+	{
+		return TEXT("0");
+	}
+	return Result;
+}
+
+FString URenderPage::PurgePageIdOrGenerateUniqueId(URenderPageCollection* PageCollection, const FString& NewPageId)
+{
+	FString Result = PurgePageIdOrReturnEmptyString(NewPageId);
+	if (Result.IsEmpty())
+	{
+		return UE::RenderPages::IRenderPagesModule::Get().GetManager().CreateUniquePageId(PageCollection);
+	}
+	return Result;
+}
+
+FString URenderPage::PurgePageName(const FString& NewPageName)
+{
+	return NewPageName.TrimStartAndEnd();
+}
+
+FString URenderPage::PurgeOutputDirectory(const FString& NewOutputDirectory)
+{
+	return UE::RenderPages::Private::FRenderPagesUtils::NormalizeOutputDirectory(FPaths::ConvertRelativePathToFull(NewOutputDirectory))
+		.Replace(*UE::RenderPages::Private::FRenderPagesUtils::NormalizeOutputDirectory(FPaths::ConvertRelativePathToFull(FPaths::ProjectDir())), TEXT("{project_dir}/"));
+}
+
+FString URenderPage::GetOutputDirectory() const
+{
+	return UE::RenderPages::Private::FRenderPagesUtils::NormalizeOutputDirectory(FPaths::ConvertRelativePathToFull(OutputDirectory.Replace(TEXT("{project_dir}"), *FPaths::ProjectDir())));
+}
+
+UMoviePipelineOutputSetting* URenderPage::GetRenderPresetOutputSettings() const
+{
+	if (!IsValid(RenderPreset))
+	{
+		return nullptr;
+	}
+	for (UMoviePipelineSetting* Settings : RenderPreset->FindSettingsByClass(UMoviePipelineOutputSetting::StaticClass(), false))
+	{
+		if (!IsValid(Settings))
+		{
+			continue;
+		}
+		if (UMoviePipelineOutputSetting* OutputSettings = Cast<UMoviePipelineOutputSetting>(Settings))
+		{
+			if (!OutputSettings->IsEnabled())
+			{
+				continue;
+			}
+			return OutputSettings;
+		}
+	}
+	return nullptr;
+}
+
+bool URenderPage::HasRemoteControlValue(const TSharedPtr<FRemoteControlEntity>& RemoteControlEntity) const
+{
+	if (!RemoteControlEntity.IsValid())
+	{
+		return false;
+	}
+
+	const FString Key = RemoteControlEntity->GetId().ToString();
+	return !!RemoteControlValues.Find(Key);
+}
+
+bool URenderPage::ConstGetRemoteControlValue(const TSharedPtr<FRemoteControlEntity>& RemoteControlEntity, TArray<uint8>& OutBinaryArray) const
+{
+	if (!RemoteControlEntity.IsValid())
+	{
+		return false;
+	}
+
+	const FString Key = RemoteControlEntity->GetId().ToString();
+	if (const FRenderPageRemoteControlPropertyData* DataPtr = RemoteControlValues.Find(Key))
+	{
+		OutBinaryArray.Append((*DataPtr).Bytes);
+		return true;
+	}
+	return false;
+}
+
+bool URenderPage::GetRemoteControlValue(const TSharedPtr<FRemoteControlEntity>& RemoteControlEntity, TArray<uint8>& OutBinaryArray)
+{
+	if (!RemoteControlEntity.IsValid())
+	{
+		return false;
+	}
+
+	const FString Key = RemoteControlEntity->GetId().ToString();
+	if (FRenderPageRemoteControlPropertyData* DataPtr = RemoteControlValues.Find(Key))
+	{
+		OutBinaryArray.Append((*DataPtr).Bytes);
+		return true;
+	}
+
+	if (!URenderPagePropRemoteControl::GetValueOfEntity(RemoteControlEntity, OutBinaryArray))
+	{
+		return false;
+	}
+	RemoteControlValues.Add(Key, FRenderPageRemoteControlPropertyData(TArray(OutBinaryArray)));
+	return true;
+}
+
+bool URenderPage::SetRemoteControlValue(const TSharedPtr<FRemoteControlEntity>& RemoteControlEntity, const TArray<uint8>& BinaryArray)
+{
+	if (!RemoteControlEntity.IsValid())
+	{
+		return false;
+	}
+	const FString Key = RemoteControlEntity->GetId().ToString();
+	RemoteControlValues.Add(Key, FRenderPageRemoteControlPropertyData(TArray(BinaryArray)));
+	return true;
+}
+
+
+URenderPageCollection::URenderPageCollection()
+	: Id(FGuid::NewGuid())
+	, PropsSourceType(ERenderPagePropsSourceType::Local)
+	, PropsSourceOrigin_RemoteControl(nullptr)
+	, CachedPropsSource(nullptr)
+	, CachedPropsSourceType(ERenderPagePropsSourceType::Local)
+{
+	if (!HasAnyFlags(RF_ClassDefaultObject))
+	{
+		LoadValuesFromCDO();
+		OnPreSaveCDO().AddUObject(this, &URenderPageCollection::SaveValuesToCDO);
+	}
+}
+
+UWorld* URenderPageCollection::GetWorld() const
+{
+	if (UWorld* CachedWorld = CachedWorldWeakPtr.Get(); IsValid(CachedWorld))
+	{
+		return CachedWorld;
+	}
+
+	if (HasAllFlags(RF_ClassDefaultObject))
+	{
+		// If we are a CDO, we must return nullptr instead of calling Outer->GetWorld() to fool UObject::ImplementsGetWorld.
+		return nullptr;
+	}
+
+	// Could be a GameInstance, could be World, could also be a WidgetTree, so we're just going to follow
+	// the outer chain to find the world we're in.
+	UObject* Outer = GetOuter();
+	while (Outer)
+	{
+		UWorld* World = Outer->GetWorld();
+		if (World)
+		{
+			CachedWorldWeakPtr = World;
+			return World;
+		}
+		Outer = Outer->GetOuter();
+	}
+	return nullptr;
+}
+
+void URenderPageCollection::PreSave(FObjectPreSaveContext SaveContext)
+{
+	if (HasAnyFlags(RF_ClassDefaultObject))
+	{
+		OnPreSaveCDO().Broadcast();
+	}
+	Super::PreSave(SaveContext);
+	SaveValuesToCDO();
+}
+
+void URenderPageCollection::PostLoad()
+{
+	Super::PostLoad();
+	LoadValuesFromCDO();
+}
+
+void URenderPageCollection::CopyValuesToOrFromCDO(const bool bToCDO)
+{
+	if (HasAnyFlags(RF_ClassDefaultObject))
+	{
+		return;
+	}
+
+	URenderPageCollection* CDO = GetCDO();
+	if (!IsValid(CDO))
+	{
+		return;
+	}
+
+	for (FProperty* Property = GetClass()->PropertyLink; Property; Property = Property->PropertyLinkNext)
+	{
+		if (!Property->HasAnyPropertyFlags(CPF_Transient | CPF_DuplicateTransient))
+		{
+			void* Data = Property->ContainerPtrToValuePtr<void>(this);
+			void* DataCDO = Property->ContainerPtrToValuePtr<void>(CDO);
+			if (bToCDO)
+			{
+				Property->CopyCompleteValue(DataCDO, Data);
+			}
+			else
+			{
+				Property->CopyCompleteValue(Data, DataCDO);
+			}
+		}
+	}
+
+	if (bToCDO)
+	{
+		CDO->RenderPages = TArray<TObjectPtr<URenderPage>>();
+		for (URenderPage* Page : RenderPages)
+		{
+			if (!IsValid(Page))
+			{
+				continue;
+			}
+			if (URenderPage* DuplicatePage = DuplicateObject(Page, CDO); IsValid(DuplicatePage))
+			{
+				CDO->RenderPages.Add(DuplicatePage);
+			}
+		}
+	}
+}
+
+void URenderPageCollection::SetPropsSource(ERenderPagePropsSourceType InPropsSourceType, UObject* InPropsSourceOrigin)
+{
+	if (InPropsSourceType == ERenderPagePropsSourceType::RemoteControl)
+	{
+		if (URemoteControlPreset* InPropsSourceOrigin_RemoteControl = Cast<URemoteControlPreset>(InPropsSourceOrigin))
+		{
+			PropsSourceType = InPropsSourceType;
+			PropsSourceOrigin_RemoteControl = InPropsSourceOrigin_RemoteControl;
+			return;
+		}
+	}
+	PropsSourceType = ERenderPagePropsSourceType::Local;
+}
+
+URenderPagePropsSourceBase* URenderPageCollection::GetPropsSource() const
+{
+	UObject* PropsSourceOrigin = GetPropsSourceOrigin();
+	if (!IsValid(CachedPropsSource) || (CachedPropsSourceType != PropsSourceType) || (CachedPropsSourceOriginWeakPtr.Get() != PropsSourceOrigin))
+	{
+		CachedPropsSourceType = PropsSourceType;
+		CachedPropsSourceOriginWeakPtr = PropsSourceOrigin;
+		CachedPropsSource = UE::RenderPages::IRenderPagesModule::Get().CreatePropsSource(const_cast<URenderPageCollection*>(this), PropsSourceType, PropsSourceOrigin);
+	}
+	return CachedPropsSource;
+}
+
+UObject* URenderPageCollection::GetPropsSourceOrigin() const
+{
+	if (PropsSourceType == ERenderPagePropsSourceType::RemoteControl)
+	{
+		return PropsSourceOrigin_RemoteControl;
+	}
+	return nullptr;
+}
+
+void URenderPageCollection::AddRenderPage(URenderPage* RenderPage)
+{
+	if (IsValid(RenderPage))
+	{
+		RenderPages.Add(RenderPage);
+	}
+}
+
+void URenderPageCollection::RemoveRenderPage(URenderPage* RenderPage)
+{
+	RenderPages.Remove(RenderPage);
+}
+
+void URenderPageCollection::InsertRenderPage(URenderPage* RenderPage, int32 Index)
+{
+	if (IsValid(RenderPage))
+	{
+		RenderPages.Insert(RenderPage, Index);
+	}
+}
+
+bool URenderPageCollection::HasRenderPage(URenderPage* RenderPage) const
+{
+	return (RenderPages.Find(RenderPage) != INDEX_NONE);
+}
+
+int32 URenderPageCollection::GetIndexOfRenderPage(URenderPage* RenderPage) const
+{
+	return RenderPages.Find(RenderPage);
+}
+
+TArray<URenderPage*> URenderPageCollection::GetRenderPages() const
+{
+	TArray<URenderPage*> Result;
+	for (URenderPage* Page : RenderPages)
+	{
+		if (!IsValid(Page))
+		{
+			continue;
+		}
+		Result.Add(Page);
+	}
+	return Result;
+}
+
+TArray<URenderPage*> URenderPageCollection::GetEnabledRenderPages() const
+{
+	TArray<URenderPage*> Result;
+	for (URenderPage* Page : RenderPages)
+	{
+		if (!IsValid(Page))
+		{
+			continue;
+		}
+		if (Page->GetIsEnabled())
+		{
+			Result.Add(Page);
+		}
+	}
+	return Result;
+}
+
+TArray<URenderPage*> URenderPageCollection::GetDisabledRenderPages() const
+{
+	TArray<URenderPage*> Result;
+	for (URenderPage* Page : RenderPages)
+	{
+		if (!IsValid(Page))
+		{
+			continue;
+		}
+		if (!Page->GetIsEnabled())
+		{
+			Result.Add(Page);
+		}
+	}
+	return Result;
+}
+
+void URenderPageCollection::InsertRenderPageBefore(URenderPage* RenderPage, URenderPage* BeforeRenderPage)
+{
+	if (IsValid(RenderPage))
+	{
+		const TArray<TObjectPtr<URenderPage>>::SizeType Index = RenderPages.Find(BeforeRenderPage);
+		if (Index == INDEX_NONE)
+		{
+			RenderPages.Add(RenderPage);
+		}
+		else
+		{
+			RenderPages.Insert(RenderPage, Index);
+		}
+	}
+}
+
+void URenderPageCollection::InsertRenderPageAfter(URenderPage* RenderPage, URenderPage* AfterRenderPage)
+{
+	if (IsValid(RenderPage))
+	{
+		const TArray<TObjectPtr<URenderPage>>::SizeType Index = RenderPages.Find(AfterRenderPage);
+		if (Index == INDEX_NONE)
+		{
+			RenderPages.Add(RenderPage);
+		}
+		else
+		{
+			RenderPages.Insert(RenderPage, Index + 1);
+		}
+	}
+}
diff --git a/Engine/Plugins/Experimental/RenderPages/Source/RenderPages/Private/RenderPage/RenderPageManager.cpp b/Engine/Plugins/Experimental/RenderPages/Source/RenderPages/Private/RenderPage/RenderPageManager.cpp
new file mode 100644
--- /dev/null
+++ b/Engine/Plugins/Experimental/RenderPages/Source/RenderPages/Private/RenderPage/RenderPageManager.cpp
@@ -0,0 +1,375 @@
+// Copyright Epic Games, Inc. All Rights Reserved.
+
+#include "RenderPage/RenderPageManager.h"
+#include "RenderPage/RenderPageMoviePipelineJob.h"
+#include "RenderPagesUtils.h"
+#include "MoviePipelineHighResSetting.h"
+#include "MoviePipelinePIEExecutor.h"
+
+
+URenderPage* UE::RenderPages::FRenderPageManager::AddNewPage(URenderPageCollection* PageCollection)
+{
+	if (!IsValid(PageCollection))
+	{
+		return nullptr;
+	}
+
+	URenderPage* Page = NewObject<URenderPage>(PageCollection);
+	Page->SetPageId(CreateUniquePageId(PageCollection));
+	Page->SetPageName(TEXT("New"));
+	Page->SetOutputDirectory(FPaths::ProjectDir() / TEXT("Saved/MovieRenders/"));
+
+	if (URenderPagePropsSourceRemoteControl* PropsSource = PageCollection->GetPropsSource<URenderPagePropsSourceRemoteControl>())
+	{
+		TArray<uint8> BinaryArray;
+		for (URenderPagePropRemoteControl* Field : PropsSource->GetProps()->GetAllCasted())
+		{
+			if (Field->GetValue(BinaryArray))
+			{
+				Page->SetRemoteControlValue(Field->GetRemoteControlEntity(), BinaryArray);
+			}
+		}
+	}
+
+	PageCollection->Modify();
+	PageCollection->AddRenderPage(Page);
+	return Page;
+}
+
+URenderPage* UE::RenderPages::FRenderPageManager::CopyPage(URenderPageCollection* PageCollection, URenderPage* Page)
+{
+	if (!IsValid(PageCollection) || !IsValid(Page))
+	{
+		return nullptr;
+	}
+
+	if (URenderPage* DuplicateRenderPage = DuplicateObject(Page, PageCollection); IsValid(DuplicateRenderPage))
+	{
+		DuplicateRenderPage->GenerateNewId();
+		DuplicateRenderPage->SetPageId(CreateUniquePageId(PageCollection));
+		DuplicateRenderPage->Modify();
+		PageCollection->Modify();
+		PageCollection->InsertRenderPageAfter(DuplicateRenderPage, Page);
+		return DuplicateRenderPage;
+	}
+	return nullptr;
+}
+
+bool UE::RenderPages::FRenderPageManager::DoesPageIdExist(URenderPageCollection* PageCollection, const FString& PageId)
+{
+	const FString PageIdToLower = PageId.ToLower();
+	for (URenderPage* Page : PageCollection->GetRenderPagesRef())
+	{
+		if (!IsValid(Page))
+		{
+			continue;
+		}
+		if (PageIdToLower == Page->GetPageId().ToLower())
+		{
+			return true;
+		}
+	}
+	return false;
+}
+
+void UE::RenderPages::FRenderPageManager::DeletePage(URenderPageCollection* PageCollection, URenderPage* Page)
+{
+	if (!IsValid(PageCollection) || !IsValid(Page))
+	{
+		return;
+	}
+
+	PageCollection->Modify();
+	PageCollection->RemoveRenderPage(Page);
+}
+
+bool UE::RenderPages::FRenderPageManager::DragDropPage(URenderPageCollection* PageCollection, URenderPage* Page, URenderPage* DroppedOnPage, const bool bAfter)
+{
+	if (!IsValid(PageCollection) || !IsValid(Page) || !IsValid(DroppedOnPage))
+	{
+		return false;
+	}
+	if (!PageCollection->HasRenderPage(Page) || !PageCollection->HasRenderPage(DroppedOnPage))
+	{
+		return false;
+	}
+
+	PageCollection->RemoveRenderPage(Page);
+	if (bAfter)
+	{
+		PageCollection->InsertRenderPageAfter(Page, DroppedOnPage);
+	}
+	else
+	{
+		PageCollection->InsertRenderPageBefore(Page, DroppedOnPage);
+	}
+	return true;
+}
+
+FString UE::RenderPages::FRenderPageManager::CreateUniquePageId(URenderPageCollection* PageCollection)
+{
+	int32 Max = 0;
+	for (URenderPage* Page : PageCollection->GetRenderPagesRef())
+	{
+		if (!IsValid(Page))
+		{
+			continue;
+		}
+		int32 Value = FCString::Atoi(*Page->GetPageId());
+		if (Value > Max)
+		{
+			Max = Value;
+		}
+	}
+	FString Result = FString::FromInt(Max + 1);
+	while (Result.Len() < GeneratedIdCharacterLength)
+	{
+		Result = TEXT("0") + Result;
+	}
+	return Result;
+}
+
+
+URenderPagesMoviePipelineRenderJob* UE::RenderPages::FRenderPageManager::CreateBatchRenderJob(URenderPageCollection* PageCollection)
+{
+	FRenderPagesMoviePipelineRenderJobCreateArgs JobArgs;
+	JobArgs.PageCollection = PageCollection;
+	JobArgs.Pages.Append(PageCollection->GetEnabledRenderPages());
+	URenderPagesMoviePipelineRenderJob* NewRenderJob = URenderPagesMoviePipelineRenderJob::Create(JobArgs);
+	if (!IsValid(NewRenderJob))
+	{
+		return nullptr;
+	}
+	return NewRenderJob;
+}
+
+
+URenderPagesMoviePipelineRenderJob* UE::RenderPages::FRenderPageManager::RenderPreviewFrame(const FRenderPageManagerRenderPreviewFrameArgs& Args)
+{
+	const FRenderPageManagerRenderPreviewFrameArgsCallback Callback = Args.Callback;
+
+	if (!IsValid(Args.Page))
+	{
+		Callback.ExecuteIfBound(false);
+		return nullptr;
+	}
+
+	URenderPage* PageCopy = DuplicateObject(Args.Page.Get(), Args.Page->GetOuter());
+	if (!IsValid(PageCopy))
+	{
+		Callback.ExecuteIfBound(false);
+		return nullptr;
+	}
+
+	PageCopy->SetPageId(PageCopy->GetId().ToString(EGuidFormats::Base36Encoded));
+
+	if (Args.Frame.IsSet())
+	{
+		constexpr int32 RenderFramesCount = 1;// can be more than 1 to prevent rendering issues, will always take the last frame that's rendered
+
+		PageCopy->SetIsCustomStartFrame(true);
+		PageCopy->SetCustomStartFrame(Args.Frame.Get(0));
+
+		PageCopy->SetIsCustomEndFrame(true);
+		PageCopy->SetCustomEndFrame(Args.Frame.Get(0));
+
+		if (!PageCopy->SetSequenceEndFrame(PageCopy->GetSequenceStartFrame().Get(0) + 1) ||
+			!PageCopy->SetSequenceStartFrame(PageCopy->GetSequenceEndFrame().Get(0) - RenderFramesCount))
+		{
+			Callback.ExecuteIfBound(false);
+			return nullptr;
+		}
+	}
+
+	PageCopy->SetIsCustomResolution(true);
+	PageCopy->SetCustomResolution(Args.Resolution);
+
+	PageCopy->SetOutputDirectory(TmpRenderedFramesPath / (Args.Frame.IsSet() ? TEXT("PreviewFrame") : TEXT("PreviewFrames")));
+
+	FRenderPagesMoviePipelineRenderJobCreateArgs JobArgs;
+	JobArgs.PageCollection = Args.PageCollection;
+	JobArgs.Pages.Add(PageCopy);
+	JobArgs.bHeadless = Args.bHeadless;
+	JobArgs.bForceOutputImage = true;
+	JobArgs.bForceOnlySingleOutput = true;
+	JobArgs.bForceUseSequenceFrameRate = Args.Frame.IsSet();
+	JobArgs.bEnsureSequentialFilenames = true;
+	JobArgs.DisableSettingsClasses.Add(UMoviePipelineAntiAliasingSetting::StaticClass());
+	JobArgs.DisableSettingsClasses.Add(UMoviePipelineHighResSetting::StaticClass());
+
+	URenderPagesMoviePipelineRenderJob* NewRenderJob = URenderPagesMoviePipelineRenderJob::Create(JobArgs);
+	if (!NewRenderJob)
+	{
+		Callback.ExecuteIfBound(false);
+		return nullptr;
+	}
+
+	const FGuid PageId = PageCopy->GetId();
+	const TOptional<int32> StartFrameOfRender = (Args.Frame.IsSet() ? TOptional<int32>() : PageCopy->GetStartFrame());
+	NewRenderJob->OnExecuteFinished().AddLambda([this,Callback,PageId,StartFrameOfRender](URenderPagesMoviePipelineRenderJob* RenderJob, const bool bSuccess)
+	{
+		if (StartFrameOfRender.IsSet())
+		{
+			if (bSuccess)
+			{
+				StartFrameOfRenders.Add(PageId, StartFrameOfRender.Get(0));
+			}
+			else
+			{
+				StartFrameOfRenders.Remove(PageId);
+			}
+		}
+		Callback.ExecuteIfBound(bSuccess);
+	});
+
+	NewRenderJob->Execute();
+	return NewRenderJob;
+}
+
+UTexture2D* UE::RenderPages::FRenderPageManager::GetSingleRenderedPreviewFrame(URenderPage* Page, UTexture2D* ReusingTexture2D, bool& bOutReusedGivenTexture2D)
+{
+	static const FString PreviewFramesDir = TmpRenderedFramesPath / TEXT("PreviewFrame");
+	if (!IsValid(Page))
+	{
+		bOutReusedGivenTexture2D = (ReusingTexture2D == nullptr);
+		return nullptr;
+	}
+	const FString PreviewFramesSubDir = Page->GetId().ToString(EGuidFormats::Base36Encoded);
+
+	TArray<FString> ImagePaths = Private::FRenderPagesUtils::GetFiles(PreviewFramesDir / PreviewFramesSubDir, true);
+	ImagePaths.Sort();
+	Algo::Reverse(ImagePaths);
+	for (const FString& ImagePath : ImagePaths)
+	{
+		if (UTexture2D* Texture = Private::FRenderPagesUtils::GetImage(ImagePath, ReusingTexture2D, bOutReusedGivenTexture2D))
+		{
+			return Texture;
+		}
+	}
+	bOutReusedGivenTexture2D = (ReusingTexture2D == nullptr);
+	return nullptr;
+}
+
+UTexture2D* UE::RenderPages::FRenderPageManager::GetSingleRenderedPreviewFrame(URenderPage* Page)
+{
+	bool bOutReusedGivenTexture2D;
+	return GetSingleRenderedPreviewFrame(Page, nullptr, bOutReusedGivenTexture2D);
+}
+
+UTexture2D* UE::RenderPages::FRenderPageManager::GetRenderedPreviewFrame(URenderPage* Page, const int32 Frame, UTexture2D* ReusingTexture2D, bool& bOutReusedGivenTexture2D)
+{
+	static const FString PreviewFramesDir = TmpRenderedFramesPath / TEXT("PreviewFrames");
+	if (!IsValid(Page))
+	{
+		bOutReusedGivenTexture2D = (ReusingTexture2D == nullptr);
+		return nullptr;
+	}
+	const FString PreviewFramesSubDir = Page->GetId().ToString(EGuidFormats::Base36Encoded);
+
+	int32 CurrentFrame = 0;
+	if (int32* StartFrameOfRenderPtr = StartFrameOfRenders.Find(Page->GetId()))
+	{
+		CurrentFrame = *StartFrameOfRenderPtr;
+	}
+
+	TArray<FString> ImagePaths = Private::FRenderPagesUtils::GetFiles(PreviewFramesDir / PreviewFramesSubDir, true);
+	for (const FString& ImagePath : ImagePaths)
+	{
+		if (!Private::FRenderPagesUtils::IsImage(ImagePath))
+		{
+			continue;
+		}
+		if (CurrentFrame == Frame)
+		{
+			return Private::FRenderPagesUtils::GetImage(ImagePath, ReusingTexture2D, bOutReusedGivenTexture2D);
+		}
+		CurrentFrame++;
+	}
+
+	bOutReusedGivenTexture2D = (ReusingTexture2D == nullptr);
+	return nullptr;
+}
+
+UTexture2D* UE::RenderPages::FRenderPageManager::GetRenderedPreviewFrame(URenderPage* Page, const int32 Frame)
+{
+	bool bOutReusedGivenTexture2D;
+	return GetRenderedPreviewFrame(Page, Frame, nullptr, bOutReusedGivenTexture2D);
+}
+
+
+void UE::RenderPages::FRenderPageManager::UpdatePagesPropValues(URenderPageCollection* PageCollection)
+{
+	if (!IsValid(PageCollection))
+	{
+		return;
+	}
+
+	URenderPagePropsSourceRemoteControl* PropsSource = PageCollection->GetPropsSource<URenderPagePropsSourceRemoteControl>();
+	if (!IsValid(PropsSource))
+	{
+		return;
+	}
+
+	TArray<URenderPage*> Pages = PageCollection->GetRenderPages();
+	TArray<uint8> BinaryArray;
+	for (URenderPagePropRemoteControl* Field : PropsSource->GetProps()->GetAllCasted())
+	{
+		if (!Field->GetValue(BinaryArray))
+		{
+			continue;
+		}
+		const TSharedPtr<FRemoteControlEntity> Entity = Field->GetRemoteControlEntity();
+		for (URenderPage* Page : Pages)
+		{
+			if (!Page->HasRemoteControlValue(Entity))
+			{
+				Page->SetRemoteControlValue(Entity, BinaryArray);
+			}
+		}
+	}
+}
+
+FRenderPageManagerPreviousPagePropValues UE::RenderPages::FRenderPageManager::ApplyPagePropValues(const URenderPageCollection* PageCollection, const URenderPage* Page)
+{
+	FRenderPageManagerPreviousPagePropValues PreviousPropValues;
+
+	if (!IsValid(PageCollection) || !IsValid(Page))
+	{
+		return PreviousPropValues;
+	}
+
+	URenderPagePropsSourceBase* PropsSource = PageCollection->GetPropsSource();
+	if (URenderPagePropsSourceRemoteControl* PropsSourceRC = Cast<URenderPagePropsSourceRemoteControl>(PropsSource))
+	{
+		for (URenderPagePropRemoteControl* Prop : PropsSourceRC->GetProps()->GetAllCasted())
+		{
+			TArray<uint8> PreviousPropData;
+			if (!Prop->GetValue(PreviousPropData))
+			{
+				continue;
+			}
+
+			TArray<uint8> PropData;
+			if (!Page->ConstGetRemoteControlValue(Prop->GetRemoteControlEntity(), PropData))
+			{
+				continue;
+			}
+
+			PreviousPropValues.RemoteControlData.Add(Prop, PreviousPropData);
+			Prop->SetValue(PropData);
+		}
+		return PreviousPropValues;
+	}
+	return PreviousPropValues;
+}
+
+void UE::RenderPages::FRenderPageManager::RestorePagePropValues(const FRenderPageManagerPreviousPagePropValues& PreviousPropValues)
+{
+	for (const auto& PropValuePair : PreviousPropValues.RemoteControlData)
+	{
+		if (IsValid(PropValuePair.Key))
+		{
+			PropValuePair.Key->SetValue(PropValuePair.Value.Bytes);
+		}
+	}
+}
diff --git a/Engine/Plugins/Experimental/RenderPages/Source/RenderPages/Private/RenderPage/RenderPageMoviePipelineJob.cpp b/Engine/Plugins/Experimental/RenderPages/Source/RenderPages/Private/RenderPage/RenderPageMoviePipelineJob.cpp
new file mode 100644
--- /dev/null
+++ b/Engine/Plugins/Experimental/RenderPages/Source/RenderPages/Private/RenderPage/RenderPageMoviePipelineJob.cpp
@@ -0,0 +1,301 @@
+// Copyright Epic Games, Inc. All Rights Reserved.
+
+#include "RenderPage/RenderPageMoviePipelineJob.h"
+
+#include "IRenderPagesModule.h"
+#include "RenderPage/RenderPageCollection.h"
+#include "RenderPage/RenderPageManager.h"
+#include "RenderPagesUtils.h"
+
+#include "LevelSequence.h"
+#include "LevelSequenceEditorModule.h"
+#include "LevelSequenceEditorModule.h"
+#include "Modules/ModuleManager.h"
+#include "MoviePipelineAntiAliasingSetting.h"
+#include "MoviePipelineDeferredPasses.h"
+#include "MoviePipelineEditorBlueprintLibrary.h"
+#include "MoviePipelineExecutor.h"
+#include "MoviePipelineImageSequenceOutput.h"
+#include "MoviePipelineOutputSetting.h"
+#include "MoviePipelinePIEExecutor.h"
+#include "MoviePipelineQueue.h"
+
+
+URenderPagesMoviePipelineRenderJob* URenderPagesMoviePipelineRenderJob::Create(const UE::RenderPages::FRenderPagesMoviePipelineRenderJobCreateArgs& Args)
+{
+	if (!IsValid(Args.PageCollection) || (Args.Pages.Num() <= 0))
+	{
+		return nullptr;
+	}
+
+	const UClass* PipelineExecutor = (IsValid(*Args.PipelineExecutorClass) ? *Args.PipelineExecutorClass : UMoviePipelinePIEExecutor::StaticClass());
+	if (!PipelineExecutor)
+	{
+		return nullptr;
+	}
+
+	URenderPagesMoviePipelineRenderJob* RenderJob = NewObject<URenderPagesMoviePipelineRenderJob>(GetTransientPackage());
+	RenderJob->PageCollection = Args.PageCollection;
+	RenderJob->RenderQueue = NewObject<UMoviePipelineQueue>(RenderJob);
+	RenderJob->ActiveExecutor = NewObject<UMoviePipelineExecutorBase>(RenderJob, PipelineExecutor);
+
+	if (Args.bHeadless)
+	{
+		if (UMoviePipelinePIEExecutor* ActiveExecutorPIE = Cast<UMoviePipelinePIEExecutor>(RenderJob->ActiveExecutor))
+		{
+			//TODO: when it becomes available on ue5-main, uncomment:  ActiveExecutorPIE->SetIsRenderingOffscreen(true);
+		}
+	}
+
+
+	for (const URenderPage* Page : Args.Pages)
+	{
+		if (!IsValid(Page))
+		{
+			continue;
+		}
+
+		ULevelSequence* PageSequence = Page->GetSequence();
+		if (!IsValid(PageSequence) || !Page->GetSequenceStartFrame().IsSet() || !Page->GetSequenceEndFrame().IsSet() || (Page->GetSequenceStartFrame().Get(0) >= Page->GetSequenceEndFrame().Get(0)))
+		{
+			RenderJob->PageStatuses.Add(Page, TEXT("Skipped"));
+			continue;
+		}
+
+		UMoviePipelineExecutorJob* NewJob = UMoviePipelineEditorBlueprintLibrary::CreateJobFromSequence(RenderJob->RenderQueue, PageSequence);
+		UMoviePipelineMasterConfig* PageRenderPreset = Page->GetRenderPreset();
+		if (IsValid(PageRenderPreset))
+		{
+			NewJob->SetConfiguration(PageRenderPreset);
+		}
+		else
+		{
+			UMoviePipelineEditorBlueprintLibrary::EnsureJobHasDefaultSettings(NewJob);
+		}
+
+		if (Args.DisableSettingsClasses.Num() > 0)
+		{
+			for (UMoviePipelineSetting* Setting : NewJob->GetConfiguration()->FindSettings<UMoviePipelineSetting>())
+			{
+				if (!IsValid(Setting))
+				{
+					continue;
+				}
+				for (TSubclassOf<UMoviePipelineSetting> DisableSettingsClass : Args.DisableSettingsClasses)
+				{
+					if (Setting->IsA(DisableSettingsClass))
+					{
+						Setting->SetIsEnabled(false);
+						break;
+					}
+				}
+			}
+		}
+
+		if (Args.bForceOutputImage || Args.bForceOnlySingleOutput)
+		{
+			bool bFound = false;
+			const bool bContainsPreferredType = IsValid(NewJob->GetConfiguration()->FindSetting<UMoviePipelineImageSequenceOutput_PNG>());
+			for (UMoviePipelineOutputBase* Setting : NewJob->GetConfiguration()->FindSettings<UMoviePipelineOutputBase>())
+			{
+				if (!IsValid(Setting))
+				{
+					continue;
+				}
+				if (Cast<UMoviePipelineImageSequenceOutputBase>(Setting))
+				{
+					if (!Args.bForceOnlySingleOutput || (!bFound && (!bContainsPreferredType || Cast<UMoviePipelineImageSequenceOutput_PNG>(Setting))))
+					{
+						bFound = true;
+						continue;
+					}
+				}
+				Setting->SetIsEnabled(false);
+			}
+			if (Args.bForceOutputImage && !bFound)
+			{
+				NewJob->GetConfiguration()->FindOrAddSettingByClass(UMoviePipelineImageSequenceOutput_PNG::StaticClass());
+			}
+		}
+
+		if (UMoviePipelineAntiAliasingSetting* ExistingAntiAliasingSettings = NewJob->GetConfiguration()->FindSetting<UMoviePipelineAntiAliasingSetting>())
+		{
+			// anti-aliasing settings already present (and enabled)
+			if ((UE::MovieRenderPipeline::GetEffectiveAntiAliasingMethod(ExistingAntiAliasingSettings) == AAM_FXAA) && (ExistingAntiAliasingSettings->SpatialSampleCount <= 1) && (ExistingAntiAliasingSettings->TemporalSampleCount <= 1))
+			{
+				ExistingAntiAliasingSettings->TemporalSampleCount = 2;// FXAA transparency fix
+			}
+		}
+		else if (UMoviePipelineAntiAliasingSetting* AntiAliasingSettings = Cast<UMoviePipelineAntiAliasingSetting>(NewJob->GetConfiguration()->FindOrAddSettingByClass(UMoviePipelineAntiAliasingSetting::StaticClass())))
+		{
+			// anti-aliasing settings not yet present (or enabled), created a new one
+			AntiAliasingSettings->EngineWarmUpCount = 0;
+			AntiAliasingSettings->RenderWarmUpCount = 0;
+			AntiAliasingSettings->SpatialSampleCount = 1;
+			AntiAliasingSettings->TemporalSampleCount = 2;
+			AntiAliasingSettings->bOverrideAntiAliasing = true;
+			AntiAliasingSettings->AntiAliasingMethod = AAM_FXAA;
+		}
+
+		bool bHasShot = false;
+		for (const UMoviePipelineExecutorShot* Shot : NewJob->ShotInfo)
+		{
+			if (!Shot)
+			{
+				continue;
+			}
+			bHasShot = true;
+
+			UMoviePipelineOutputSetting* Setting = Cast<UMoviePipelineOutputSetting>(UMoviePipelineBlueprintLibrary::FindOrGetDefaultSettingForShot(UMoviePipelineOutputSetting::StaticClass(), NewJob->GetConfiguration(), Shot));
+
+			Setting->bUseCustomPlaybackRange = true;
+			Setting->CustomStartFrame = Page->GetSequenceStartFrame().Get(0);
+			Setting->CustomEndFrame = Page->GetSequenceEndFrame().Get(0);
+
+			if (Args.bForceUseSequenceFrameRate)
+			{
+				Setting->bUseCustomFrameRate = false;
+			}
+
+			if (Page->GetIsCustomResolution())
+			{
+				Setting->OutputResolution = Page->GetCustomResolution();
+			}
+
+			const FString PageOutputRootDirectory = Page->GetOutputDirectory();
+			const FString PageId = Page->GetPageId();
+			if (!PageOutputRootDirectory.IsEmpty() && !PageId.IsEmpty())
+			{
+				const FString PageOutputDirectory = PageOutputRootDirectory / PageId;
+				UE::RenderPages::Private::FRenderPagesUtils::DeleteDirectory(PageOutputDirectory);
+				Setting->OutputDirectory.Path = PageOutputDirectory;
+			}
+
+			if (Args.bEnsureSequentialFilenames || !IsValid(PageRenderPreset))
+			{
+				Setting->FileNameFormat = TEXT("{frame_number}");
+				Setting->ZeroPadFrameNumbers = 10;
+				Setting->FrameNumberOffset = 1000000000;
+			}
+		}
+		if (!bHasShot)
+		{
+			RenderJob->PageStatuses.Add(Page, TEXT("Skipped"));
+		}
+
+		RenderJob->PageExecutorJobs.Add(Page, NewJob);
+	}
+
+	if (RenderJob->PageExecutorJobs.Num() <= 0)
+	{
+		return nullptr;
+	}
+	return RenderJob;
+}
+
+
+void URenderPagesMoviePipelineRenderJob::Execute()
+{
+	if (ActiveExecutor->IsRendering())
+	{
+		return;
+	}
+	OnExecuteStartedDelegate.Broadcast(this);
+	AddToRoot();
+	if (ILevelSequenceEditorModule* LevelSequenceEditorModule = FModuleManager::GetModulePtr<ILevelSequenceEditorModule>("LevelSequenceEditor"); LevelSequenceEditorModule)
+	{
+		LevelSequenceEditorModule->OnComputePlaybackContext().AddUObject(this, &URenderPagesMoviePipelineRenderJob::ComputePlaybackContext);
+	}
+	if (UMoviePipelinePIEExecutor* ActiveExecutorPIE = Cast<UMoviePipelinePIEExecutor>(ActiveExecutor))
+	{
+		ActiveExecutorPIE->OnIndividualJobStarted().AddUObject(this, &URenderPagesMoviePipelineRenderJob::ExecutePageStarted);
+		ActiveExecutorPIE->OnIndividualJobWorkFinished().AddUObject(this, &URenderPagesMoviePipelineRenderJob::ExecutePageFinished);
+	}
+	ActiveExecutor->OnExecutorFinished().AddUObject(this, &URenderPagesMoviePipelineRenderJob::ExecuteFinished);
+	ActiveExecutor->Execute(RenderQueue);
+}
+
+void URenderPagesMoviePipelineRenderJob::Cancel()
+{
+	if (ActiveExecutor->IsRendering())
+	{
+		ActiveExecutor->CancelAllJobs();
+	}
+}
+
+FString URenderPagesMoviePipelineRenderJob::GetPageStatus(URenderPage* Page) const
+{
+	if (!IsValid(Page))
+	{
+		return TEXT("");
+	}
+
+	const FString* StatusPtr = PageStatuses.Find(Page);
+	if (StatusPtr)
+	{
+		// got the result from the Page Statuses map
+		return *StatusPtr;
+	}
+
+	const TObjectPtr<UMoviePipelineExecutorJob>* JobPtr = PageExecutorJobs.Find(Page);
+	if (!JobPtr)
+	{
+		return TEXT("");
+	}
+	TObjectPtr<UMoviePipelineExecutorJob> Job = *JobPtr;
+	if (!Job)
+	{
+		return TEXT("");
+	}
+
+	if (Job->GetStatusProgress() > 0)
+	{
+		return TEXT("Rendering");
+	}
+	return Job->GetStatusMessage();
+}
+
+void URenderPagesMoviePipelineRenderJob::ComputePlaybackContext(bool& bOutAllowBinding)
+{
+	bOutAllowBinding = false;
+}
+
+void URenderPagesMoviePipelineRenderJob::ExecutePageStarted(UMoviePipelineExecutorJob* JobToStart)
+{
+	const TObjectPtr<const URenderPage>* PagePtr = PageExecutorJobs.FindKey(JobToStart);
+	if (!PagePtr)
+	{
+		return;
+	}
+	TObjectPtr<const URenderPage> Page = *PagePtr;
+	PageStatuses.Add(Page, TEXT("Rendering..."));
+	PreviousPageProps = UE::RenderPages::IRenderPagesModule::Get().GetManager().ApplyPagePropValues(PageCollection, Page.Get());
+}
+
+void URenderPagesMoviePipelineRenderJob::ExecutePageFinished(FMoviePipelineOutputData PipelineOutputData)
+{
+	const TObjectPtr<const URenderPage>* PagePtr = PageExecutorJobs.FindKey(PipelineOutputData.Job);
+	if (!PagePtr)
+	{
+		return;
+	}
+	TObjectPtr<const URenderPage> Page = *PagePtr;
+	PageStatuses.Add(Page, TEXT("Done"));
+	UE::RenderPages::IRenderPagesModule::Get().GetManager().RestorePagePropValues(PreviousPageProps);
+	PreviousPageProps = FRenderPageManagerPreviousPagePropValues();
+}
+
+void URenderPagesMoviePipelineRenderJob::ExecuteFinished(UMoviePipelineExecutorBase* PipelineExecutor, const bool bSuccess)
+{
+	if (ILevelSequenceEditorModule* LevelSequenceEditorModule = FModuleManager::GetModulePtr<ILevelSequenceEditorModule>("LevelSequenceEditorModule"))
+	{
+		LevelSequenceEditorModule->OnComputePlaybackContext().RemoveAll(this);
+	}
+	if (UMoviePipelinePIEExecutor* PipelineExecutorPIE = Cast<UMoviePipelinePIEExecutor>(PipelineExecutor))
+	{
+		PipelineExecutorPIE->OnIndividualJobStarted().RemoveAll(this);
+		PipelineExecutorPIE->OnIndividualJobWorkFinished().RemoveAll(this);
+	}
+	RemoveFromRoot();
+	OnExecuteFinishedDelegate.Broadcast(this, bSuccess);
+}
diff --git a/Engine/Plugins/Experimental/RenderPages/Source/RenderPages/Private/RenderPage/RenderPagePropsSource.cpp b/Engine/Plugins/Experimental/RenderPages/Source/RenderPages/Private/RenderPage/RenderPagePropsSource.cpp
new file mode 100644
--- /dev/null
+++ b/Engine/Plugins/Experimental/RenderPages/Source/RenderPages/Private/RenderPage/RenderPagePropsSource.cpp
@@ -0,0 +1,159 @@
+// Copyright Epic Games, Inc. All Rights Reserved.
+
+#include "RenderPage/RenderPagePropsSource.h"
+#include "RenderPagesLog.h"
+#include "IRemoteControlModule.h"
+#include "Backends/JsonStructDeserializerBackend.h"
+#include "Backends/JsonStructSerializerBackend.h"
+
+
+bool /*URenderPagePropRemoteControl::*/GetObjectRef(const TSharedPtr<FRemoteControlProperty>& Field, const ERCAccess Access, FRCObjectReference& OutObjectRef)
+{
+	FRCObjectReference ObjectRef;
+	FString* ErrorText = nullptr;
+	if (!IRemoteControlModule::Get().ResolveObjectProperty(Access, Field->GetBoundObject(), Field->FieldPathInfo, ObjectRef, ErrorText))
+	{
+		UE_LOG(LogRenderPages, Warning, TEXT("Couldn\'t resolve object property \"%s\" in \"%s\": %s"), *Field->FieldName.ToString(), (ErrorText ? **ErrorText : TEXT("unknown")));
+		return false;
+	}
+	OutObjectRef = ObjectRef;
+	return true;
+}
+
+bool URenderPagePropRemoteControl::GetValueOfEntity(const TSharedPtr<FRemoteControlEntity>& RemoteControlEntity, TArray<uint8>& OutBinaryArray)
+{
+	OutBinaryArray.Empty();
+
+	TSharedPtr<FRemoteControlProperty> Field = StaticCastSharedPtr<FRemoteControlProperty>(RemoteControlEntity);
+	if (!Field.IsValid())
+	{
+		return false;
+	}
+
+	FRCObjectReference ObjectRef;
+	if (!GetObjectRef(Field, ERCAccess::READ_ACCESS, ObjectRef))
+	{
+		return false;
+	}
+	FMemoryWriter Writer = FMemoryWriter(OutBinaryArray);
+	FJsonStructSerializerBackend WriterBackend = FJsonStructSerializerBackend(Writer, EStructSerializerBackendFlags::Default);
+	return IRemoteControlModule::Get().GetObjectProperties(ObjectRef, WriterBackend);
+}
+
+bool URenderPagePropRemoteControl::SetValueOfEntity(const TSharedPtr<FRemoteControlEntity>& RemoteControlEntity, const TArray<uint8>& BinaryArray)
+{
+	TSharedPtr<FRemoteControlProperty> Field = StaticCastSharedPtr<FRemoteControlProperty>(RemoteControlEntity);
+	if (!Field.IsValid())
+	{
+		return false;
+	}
+
+	FRCObjectReference ObjectRefRead;
+	if (GetObjectRef(Field, ERCAccess::READ_ACCESS, ObjectRefRead))
+	{
+		TArray<uint8> CurrentBinaryArray;
+		FMemoryWriter Writer = FMemoryWriter(CurrentBinaryArray);
+		FJsonStructSerializerBackend WriterBackend = FJsonStructSerializerBackend(Writer, EStructSerializerBackendFlags::Default);
+		if (IRemoteControlModule::Get().GetObjectProperties(ObjectRefRead, WriterBackend))
+		{
+			if (CurrentBinaryArray == BinaryArray)
+			{
+				// if the given value is already set, don't do anything
+				return true;
+			}
+		}
+	}
+
+	FRCObjectReference ObjectRefWrite;
+	if (!GetObjectRef(Field, ERCAccess::WRITE_ACCESS, ObjectRefWrite))
+	{
+		return false;
+	}
+	FMemoryReader Reader = FMemoryReader(BinaryArray);
+	FJsonStructDeserializerBackend ReaderBackend = FJsonStructDeserializerBackend(Reader);
+	return IRemoteControlModule::Get().SetObjectProperties(ObjectRefWrite, ReaderBackend, ERCPayloadType::Json);
+}
+
+void URenderPagePropRemoteControl::Initialize(const TSharedPtr<FRemoteControlEntity>& InRemoteControlEntity)
+{
+	RemoteControlEntity = InRemoteControlEntity;
+}
+
+
+void URenderPagePropsRemoteControl::Initialize(URemoteControlPreset* InRemoteControlPreset)
+{
+	RemoteControlPreset = InRemoteControlPreset;
+}
+
+TArray<URenderPagePropBase*> URenderPagePropsRemoteControl::GetAll() const
+{
+	TArray<URenderPagePropBase*> Result;
+	if (IsValid(RemoteControlPreset))
+	{
+		for (const TWeakPtr<FRemoteControlEntity>& PropWeakPtr : RemoteControlPreset->GetExposedEntities<FRemoteControlEntity>())
+		{
+			if (const TSharedPtr<FRemoteControlEntity> Prop = PropWeakPtr.Pin())
+			{
+				URenderPagePropRemoteControl* PropObj = NewObject<URenderPagePropRemoteControl>(const_cast<URenderPagePropsRemoteControl*>(this));
+				PropObj->Initialize(Prop);
+				Result.Add(PropObj);
+			}
+		}
+	}
+	return Result;
+}
+
+TArray<URenderPagePropRemoteControl*> URenderPagePropsRemoteControl::GetAllCasted() const
+{
+	TArray<URenderPagePropRemoteControl*> Result;
+	if (IsValid(RemoteControlPreset))
+	{
+		for (const TWeakPtr<FRemoteControlEntity>& PropWeakPtr : RemoteControlPreset->GetExposedEntities<FRemoteControlEntity>())
+		{
+			if (const TSharedPtr<FRemoteControlEntity> Prop = PropWeakPtr.Pin())
+			{
+				URenderPagePropRemoteControl* PropObj = NewObject<URenderPagePropRemoteControl>(const_cast<URenderPagePropsRemoteControl*>(this));
+				PropObj->Initialize(Prop);
+				Result.Add(PropObj);
+			}
+		}
+	}
+	return Result;
+}
+
+
+void URenderPagePropsSourceRemoteControl::SetSourceOrigin(UObject* SourceOrigin)
+{
+	if (!IsValid(SourceOrigin))
+	{
+		RemoteControlPreset = nullptr;
+		return;
+	}
+
+	RemoteControlPreset = Cast<URemoteControlPreset>(SourceOrigin);
+	if (IsValid(RemoteControlPreset))
+	{
+		ActivePresetGroup = RemoteControlPreset->Layout.GetDefaultGroup().Name;
+	}
+}
+
+URenderPagePropsRemoteControl* URenderPagePropsSourceRemoteControl::GetProps() const
+{
+	URenderPagePropsRemoteControl* PropsObj = NewObject<URenderPagePropsRemoteControl>(const_cast<URenderPagePropsSourceRemoteControl*>(this));
+	PropsObj->Initialize(IsValid(RemoteControlPreset) ? RemoteControlPreset : nullptr);
+	return PropsObj;
+}
+
+void URenderPagePropsSourceRemoteControl::GetAvailablePresetGroups(TArray<FName>& OutPresetGroups) const
+{
+	OutPresetGroups.Reset();
+
+	if (IsValid(RemoteControlPreset))
+	{
+		TArray<FRemoteControlPresetGroup> PresetGroups = RemoteControlPreset->Layout.GetGroups();
+		for (const FRemoteControlPresetGroup& PresetGroup : PresetGroups)
+		{
+			OutPresetGroups.Add(PresetGroup.Name);
+		}
+	}
+}
diff --git a/Engine/Plugins/Experimental/RenderPages/Source/RenderPages/Private/RenderPage/RenderPagesBlueprintGeneratedClass.cpp b/Engine/Plugins/Experimental/RenderPages/Source/RenderPages/Private/RenderPage/RenderPagesBlueprintGeneratedClass.cpp
new file mode 100644
--- /dev/null
+++ b/Engine/Plugins/Experimental/RenderPages/Source/RenderPages/Private/RenderPage/RenderPagesBlueprintGeneratedClass.cpp
@@ -0,0 +1,14 @@
+// Copyright Epic Games, Inc. All Rights Reserved.
+
+#include "RenderPage/RenderPagesBlueprintGeneratedClass.h"
+
+
+uint8* URenderPagesBlueprintGeneratedClass::GetPersistentUberGraphFrame(UObject* Obj, UFunction* FuncToCheck) const
+{
+	if (!IsInGameThread())
+	{
+		// we cant use the persistent frame if we are executing in parallel (as we could potentially thunk to BP)
+		return nullptr;
+	}
+	return Super::GetPersistentUberGraphFrame(Obj, FuncToCheck);
+}
diff --git a/Engine/Plugins/Experimental/RenderPages/Source/RenderPages/Private/RenderPagesModule.cpp b/Engine/Plugins/Experimental/RenderPages/Source/RenderPages/Private/RenderPagesModule.cpp
new file mode 100644
--- /dev/null
+++ b/Engine/Plugins/Experimental/RenderPages/Source/RenderPages/Private/RenderPagesModule.cpp
@@ -0,0 +1,87 @@
+// Copyright Epic Games, Inc. All Rights Reserved.
+
+#include "RenderPagesModule.h"
+#include "RenderPage/RenderPageManager.h"
+#include "RenderPage/RenderPagePropsSource.h"
+#include "Factories/RenderPagePropsSourceFactoryLocal.h"
+#include "Factories/RenderPagePropsSourceFactoryRemoteControl.h"
+
+#include "HAL/FileManager.h"
+
+#define LOCTEXT_NAMESPACE "RenderPagesModule"
+
+
+void UE::RenderPages::Private::FRenderPagesModule::StartupModule()
+{
+	IFileManager::Get().DeleteDirectory(*FRenderPageManager::TmpRenderedFramesPath, false, true);
+	RegisterPropsSourceFactories();
+	CreateManager();
+}
+
+void UE::RenderPages::Private::FRenderPagesModule::ShutdownModule()
+{
+	RemoveManager();
+	UnregisterPropsSourceFactories();
+}
+
+
+void UE::RenderPages::Private::FRenderPagesModule::CreateManager()
+{
+	Manager = MakeUnique<FRenderPageManager>();
+}
+
+void UE::RenderPages::Private::FRenderPagesModule::RemoveManager()
+{
+	Manager.Reset();
+}
+
+UE::RenderPages::FRenderPageManager& UE::RenderPages::Private::FRenderPagesModule::GetManager() const
+{
+	check(Manager.IsValid());
+	return *Manager;
+}
+
+
+void UE::RenderPages::Private::FRenderPagesModule::RegisterPropsSourceFactories()
+{
+	RegisterPropsSourceFactory(ERenderPagePropsSourceType::Local, MakeShared<FRenderPagePropsSourceFactoryLocal>());
+	RegisterPropsSourceFactory(ERenderPagePropsSourceType::RemoteControl, MakeShared<FRenderPagePropsSourceFactoryRemoteControl>());
+}
+
+void UE::RenderPages::Private::FRenderPagesModule::UnregisterPropsSourceFactories()
+{
+	UnregisterPropsSourceFactory(ERenderPagePropsSourceType::Local);
+	UnregisterPropsSourceFactory(ERenderPagePropsSourceType::RemoteControl);
+}
+
+void UE::RenderPages::Private::FRenderPagesModule::RegisterPropsSourceFactory(const ERenderPagePropsSourceType PropsSourceType, const TSharedPtr<IRenderPagePropsSourceFactory>& InFactory)
+{
+	PropsSourceFactories.Add(PropsSourceType, InFactory);
+}
+
+void UE::RenderPages::Private::FRenderPagesModule::UnregisterPropsSourceFactory(const ERenderPagePropsSourceType PropsSourceType)
+{
+	PropsSourceFactories.Remove(PropsSourceType);
+}
+
+URenderPagePropsSourceBase* UE::RenderPages::Private::FRenderPagesModule::CreatePropsSource(UObject* Outer, ERenderPagePropsSourceType PropsSourceType, UObject* PropsSourceOrigin)
+{
+	TSharedPtr<IRenderPagePropsSourceFactory>* FactoryPtr = PropsSourceFactories.Find(PropsSourceType);
+	if (!FactoryPtr)
+	{
+		return nullptr;
+	}
+
+	TSharedPtr<IRenderPagePropsSourceFactory> Factory = *FactoryPtr;
+	if (!Factory)
+	{
+		return nullptr;
+	}
+
+	return Factory->CreateInstance(Outer, PropsSourceOrigin);
+}
+
+
+#undef LOCTEXT_NAMESPACE
+
+IMPLEMENT_MODULE(UE::RenderPages::Private::FRenderPagesModule, RenderPages)
diff --git a/Engine/Plugins/Experimental/RenderPages/Source/RenderPages/Private/RenderPagesModule.h b/Engine/Plugins/Experimental/RenderPages/Source/RenderPages/Private/RenderPagesModule.h
new file mode 100644
--- /dev/null
+++ b/Engine/Plugins/Experimental/RenderPages/Source/RenderPages/Private/RenderPagesModule.h
@@ -0,0 +1,43 @@
+// Copyright Epic Games, Inc. All Rights Reserved.
+
+#pragma once
+
+#include "CoreMinimal.h"
+#include "IRenderPagesModule.h"
+#include "Factories/IRenderPagePropsSourceFactory.h"
+
+
+namespace UE::RenderPages::Private
+{
+	/**
+	 * The implementation of the IRenderPagesModule interface.
+	 */
+	class FRenderPagesModule : public IRenderPagesModule
+	{
+	public:
+		//~ Begin IModuleInterface interface
+		virtual void StartupModule() override;
+		virtual void ShutdownModule() override;
+		//~ End IModuleInterface interface
+
+		//~ Begin IRenderPagesModule interface
+		virtual FRenderPageManager& GetManager() const override;
+
+		virtual URenderPagePropsSourceBase* CreatePropsSource(UObject* Outer, ERenderPagePropsSourceType PropsSourceType, UObject* PropsSourceOrigin) override;
+		virtual const TMap<ERenderPagePropsSourceType, TSharedPtr<IRenderPagePropsSourceFactory>>& GetPropsSourceFactories() const override { return PropsSourceFactories; }
+		//~ End IRenderPagesModule interface
+
+	private:
+		void CreateManager();
+		void RemoveManager();
+
+		void RegisterPropsSourceFactories();
+		void UnregisterPropsSourceFactories();
+		void RegisterPropsSourceFactory(const ERenderPagePropsSourceType PropsSourceType, const TSharedPtr<IRenderPagePropsSourceFactory>& InFactory);
+		void UnregisterPropsSourceFactory(const ERenderPagePropsSourceType PropsSourceType);
+
+	private:
+		TUniquePtr<FRenderPageManager> Manager;
+		TMap<ERenderPagePropsSourceType, TSharedPtr<IRenderPagePropsSourceFactory>> PropsSourceFactories;
+	};
+}
diff --git a/Engine/Plugins/Experimental/RenderPages/Source/RenderPages/Private/RenderPagesUtils.cpp b/Engine/Plugins/Experimental/RenderPages/Source/RenderPages/Private/RenderPagesUtils.cpp
new file mode 100644
--- /dev/null
+++ b/Engine/Plugins/Experimental/RenderPages/Source/RenderPages/Private/RenderPagesUtils.cpp
@@ -0,0 +1,250 @@
+// Copyright Epic Games, Inc. All Rights Reserved.
+
+#include "RenderPagesUtils.h"
+#include "Engine/Texture2D.h"
+#include "HAL/FileManager.h"
+#include "HAL/PlatformFileManager.h"
+#include "IImageWrapper.h"
+#include "IImageWrapperModule.h"
+#include "ImageUtils.h"
+#include "Misc/Paths.h"
+#include "Modules/ModuleManager.h"
+
+
+bool UE::RenderPages::Private::FRenderPagesUtils::IsImage(const FString& ImagePath)
+{
+	static IImageWrapperModule* ImageWrapperModule = FModuleManager::LoadModulePtr<IImageWrapperModule>("ImageWrapper");
+	if (ImageWrapperModule == nullptr)
+	{
+		return false;
+	}
+
+	EImageFormat ImageFormat = ImageWrapperModule->GetImageFormatFromExtension(*ImagePath);
+	if (ImageFormat == EImageFormat::Invalid)
+	{
+		return false;
+	}
+
+	return true;
+}
+
+UTexture2D* UE::RenderPages::Private::FRenderPagesUtils::GetImage(const FString& ImagePath, UTexture2D* Texture2D, bool& bOutReusedGivenTexture2D)
+{
+	static IImageWrapperModule* ImageWrapperModule = FModuleManager::LoadModulePtr<IImageWrapperModule>("ImageWrapper");
+	if (ImageWrapperModule == nullptr)
+	{
+		bOutReusedGivenTexture2D = (Texture2D == nullptr);
+		return nullptr;
+	}
+
+	EImageFormat ImageFormat = ImageWrapperModule->GetImageFormatFromExtension(*ImagePath);
+	if (ImageFormat == EImageFormat::Invalid)
+	{
+		bOutReusedGivenTexture2D = (Texture2D == nullptr);
+		return nullptr;
+	}
+
+	TArray<uint8> ImageBytes = GetFileData(ImagePath);
+	if (ImageBytes.Num() <= 0)
+	{
+		bOutReusedGivenTexture2D = (Texture2D == nullptr);
+		return nullptr;
+	}
+
+	return BytesToExistingImage(ImageBytes, ImageFormat, Texture2D, bOutReusedGivenTexture2D);
+}
+
+
+UTexture2D* UE::RenderPages::Private::FRenderPagesUtils::BytesToImage(const TArray<uint8>& ByteArray, const EImageFormat ImageFormat)
+{
+	static IImageWrapperModule* ImageWrapperModule = FModuleManager::LoadModulePtr<IImageWrapperModule>("ImageWrapper");
+	if (ImageWrapperModule == nullptr)
+	{
+		return nullptr;
+	}
+
+	TSharedPtr<IImageWrapper> ImageWrapper = ImageWrapperModule->CreateImageWrapper(ImageFormat);
+	if (!ImageWrapper.IsValid() || !ImageWrapper->SetCompressed(ByteArray.GetData(), ByteArray.Num()))
+	{
+		return nullptr;
+	}
+
+	TArray<uint8> Uncompressed;
+	if (!ImageWrapper->GetRaw(ERGBFormat::BGRA, 8, Uncompressed))
+	{
+		return nullptr;
+	}
+
+	return DataToTexture2D(ImageWrapper->GetWidth(), ImageWrapper->GetHeight(), Uncompressed.GetData(), Uncompressed.Num());
+}
+
+UTexture2D* UE::RenderPages::Private::FRenderPagesUtils::BytesToExistingImage(const TArray<uint8>& ByteArray, const EImageFormat ImageFormat, UTexture2D* Texture2D, bool& bOutReusedGivenTexture2D)
+{
+	static IImageWrapperModule* ImageWrapperModule = FModuleManager::LoadModulePtr<IImageWrapperModule>("ImageWrapper");
+	if (ImageWrapperModule == nullptr)
+	{
+		bOutReusedGivenTexture2D = (Texture2D == nullptr);
+		return nullptr;
+	}
+
+	TSharedPtr<IImageWrapper> ImageWrapper = ImageWrapperModule->CreateImageWrapper(ImageFormat);
+	if (!ImageWrapper.IsValid() || !ImageWrapper->SetCompressed(ByteArray.GetData(), ByteArray.Num()))
+	{
+		bOutReusedGivenTexture2D = (Texture2D == nullptr);
+		return nullptr;
+	}
+
+	TArray<uint8> Uncompressed;
+	if (!ImageWrapper->GetRaw(ERGBFormat::BGRA, 8, Uncompressed))
+	{
+		bOutReusedGivenTexture2D = (Texture2D == nullptr);
+		return nullptr;
+	}
+
+	UTexture2D* NewTexture2D = DataToExistingTexture2D(ImageWrapper->GetWidth(), ImageWrapper->GetHeight(), Uncompressed.GetData(), Uncompressed.Num(), Texture2D, bOutReusedGivenTexture2D);
+	bOutReusedGivenTexture2D = (Texture2D == NewTexture2D);
+	return NewTexture2D;
+}
+
+
+UTexture2D* UE::RenderPages::Private::FRenderPagesUtils::DataToTexture2D(int32 Width, int32 Height, const void* Src, SIZE_T Count)
+{
+	UTexture2D* Texture2D = UTexture2D::CreateTransient(Width, Height, PF_B8G8R8A8);
+	if (Texture2D == nullptr)
+	{
+		return nullptr;
+	}
+	Texture2D->bNoTiling = true;
+
+#if WITH_EDITORONLY_DATA
+	Texture2D->MipGenSettings = TMGS_NoMipmaps;
+#endif
+
+	void* TextureData = Texture2D->GetPlatformData()->Mips[0].BulkData.Lock(LOCK_READ_WRITE);
+	FMemory::Memcpy(TextureData, Src, Count);
+	Texture2D->GetPlatformData()->Mips[0].BulkData.Unlock();
+
+	Texture2D->UpdateResource();
+	return Texture2D;
+}
+
+UTexture2D* UE::RenderPages::Private::FRenderPagesUtils::DataToExistingTexture2D(int32 Width, int32 Height, const void* Src, SIZE_T Count, UTexture2D* Texture2D, bool& bOutReusedGivenTexture2D)
+{
+	UTexture2D* OriginalTexture2D = Texture2D;
+
+	bOutReusedGivenTexture2D = true;
+	if ((Texture2D == nullptr) || (Texture2D->GetSizeX() != Width) || (Texture2D->GetSizeY() != Height) || (Texture2D->GetPixelFormat() != PF_B8G8R8A8))
+	{
+		bOutReusedGivenTexture2D = false;
+		Texture2D = UTexture2D::CreateTransient(Width, Height, PF_B8G8R8A8);
+	}
+	if (Texture2D == nullptr)
+	{
+		bOutReusedGivenTexture2D = (OriginalTexture2D == nullptr);
+		return nullptr;
+	}
+	Texture2D->bNoTiling = true;
+
+#if WITH_EDITORONLY_DATA
+	Texture2D->MipGenSettings = TMGS_NoMipmaps;
+#endif
+
+	void* TextureData = Texture2D->GetPlatformData()->Mips[0].BulkData.Lock(LOCK_READ_WRITE);
+	FMemory::Memcpy(TextureData, Src, Count);
+	Texture2D->GetPlatformData()->Mips[0].BulkData.Unlock();
+
+	Texture2D->UpdateResource();
+	return Texture2D;
+}
+
+
+TArray<FString> UE::RenderPages::Private::FRenderPagesUtils::GetFiles(const FString& Directory, const bool bRecursive)
+{
+	class FRenderPagesUtilsGetFilesVisitor : public IPlatformFile::FDirectoryVisitor
+	{
+	public:
+		TArray<FString> Files;
+		virtual bool Visit(const TCHAR* FileName, const bool bIsDirectory) override
+		{
+			if (!bIsDirectory)
+			{
+				Files.Add(FileName);
+			}
+			return true;
+		}
+	};
+	FRenderPagesUtilsGetFilesVisitor GetFilesVisitor;
+	if (bRecursive)
+	{
+		FPlatformFileManager::Get().GetPlatformFile().IterateDirectoryRecursively(*Directory, GetFilesVisitor);
+	}
+	else
+	{
+		FPlatformFileManager::Get().GetPlatformFile().IterateDirectory(*Directory, GetFilesVisitor);
+	}
+	return GetFilesVisitor.Files;
+}
+
+TArray<uint8> UE::RenderPages::Private::FRenderPagesUtils::GetFileData(const FString& File)
+{
+	TArray<uint8> Data;
+	IFileHandle* FileHandle = FPlatformFileManager::Get().GetPlatformFile().OpenRead(*File);
+	if (FileHandle != nullptr)
+	{
+		int64 FileSize = FileHandle->Size();
+		if (FileSize > 0x7fffffff)
+		{
+			FileSize = 0;
+		}
+		if (FileSize > 0)
+		{
+			Data.SetNum(FileSize);
+			FileHandle->Read(Data.GetData(), FileSize);
+		}
+		delete FileHandle;
+	}
+	return Data;
+}
+
+
+void UE::RenderPages::Private::FRenderPagesUtils::DeleteDirectory(const FString& Directory)
+{
+	IFileManager::Get().DeleteDirectory(*Directory, false, true);
+}
+
+void UE::RenderPages::Private::FRenderPagesUtils::EmptyDirectory(const FString& Directory)
+{
+	class FRenderPagesUtilsEmptyDirectoryVisitor : public IPlatformFile::FDirectoryVisitor
+	{
+	public:
+		virtual bool Visit(const TCHAR* FileName, const bool bIsDirectory) override
+		{
+			if (bIsDirectory)
+			{
+				IFileManager::Get().DeleteDirectory(FileName, false, false);
+			}
+			else
+			{
+				IFileManager::Get().Delete(FileName, false, true, true);
+			}
+			return true;
+		}
+	};
+	FRenderPagesUtilsEmptyDirectoryVisitor GetFilesVisitor;
+	FPlatformFileManager::Get().GetPlatformFile().IterateDirectoryRecursively(*Directory, GetFilesVisitor);
+}
+
+
+FString UE::RenderPages::Private::FRenderPagesUtils::NormalizeOutputDirectory(const FString& OutputDirectory)
+{
+	FString NewOutputDirectory = OutputDirectory;
+	FPaths::NormalizeDirectoryName(NewOutputDirectory);
+	FPaths::RemoveDuplicateSlashes(NewOutputDirectory);
+	FPaths::CollapseRelativeDirectories(NewOutputDirectory);
+	FPaths::RemoveDuplicateSlashes(NewOutputDirectory);
+	if (!NewOutputDirectory.EndsWith(TEXT("/"), ESearchCase::CaseSensitive))
+	{
+		NewOutputDirectory += TEXT("/");
+	}
+	return NewOutputDirectory;
+}
diff --git a/Engine/Plugins/Experimental/RenderPages/Source/RenderPages/Public/Factories/IRenderPagePropsSourceFactory.h b/Engine/Plugins/Experimental/RenderPages/Source/RenderPages/Public/Factories/IRenderPagePropsSourceFactory.h
new file mode 100644
--- /dev/null
+++ b/Engine/Plugins/Experimental/RenderPages/Source/RenderPages/Public/Factories/IRenderPagePropsSourceFactory.h
@@ -0,0 +1,22 @@
+// Copyright Epic Games, Inc. All Rights Reserved.
+
+#pragma once
+
+#include "CoreMinimal.h"
+
+
+class URenderPagePropsSourceBase;
+
+
+namespace UE::RenderPages
+{
+	/**
+	 * The base class for the factory classes that will create URenderPagePropsSourceBase instances.
+	 */
+	class RENDERPAGES_API IRenderPagePropsSourceFactory
+	{
+	public:
+		virtual ~IRenderPagePropsSourceFactory() = default;
+		virtual URenderPagePropsSourceBase* CreateInstance(UObject* Outer, UObject* PropsSourceOrigin) { return nullptr; }
+	};
+}
diff --git a/Engine/Plugins/Experimental/RenderPages/Source/RenderPages/Public/IRenderPagesModule.h b/Engine/Plugins/Experimental/RenderPages/Source/RenderPages/Public/IRenderPagesModule.h
new file mode 100644
--- /dev/null
+++ b/Engine/Plugins/Experimental/RenderPages/Source/RenderPages/Public/IRenderPagesModule.h
@@ -0,0 +1,51 @@
+// Copyright Epic Games, Inc. All Rights Reserved.
+
+#pragma once
+
+#include "CoreMinimal.h"
+#include "Modules/ModuleManager.h"
+
+
+class URenderPagePropsSourceBase;
+enum class ERenderPagePropsSourceType : uint8;
+
+namespace UE::RenderPages
+{
+	class FRenderPageManager;
+	class IRenderPagePropsSourceFactory;
+}
+
+
+namespace UE::RenderPages
+{
+	/**
+	* RenderPages module interface.
+	*/
+	class IRenderPagesModule : public IModuleInterface
+	{
+	public:
+		/**
+		 * Singleton-like access to IRenderPagesModule.
+		 *
+		 * @return Returns RenderPagesModule singleton instance, loading the module on demand if needed.
+		 */
+		static IRenderPagesModule& Get()
+		{
+			static const FName ModuleName = "RenderPages";
+			return FModuleManager::LoadModuleChecked<IRenderPagesModule>(ModuleName);
+		}
+
+		/**
+		 * Singleton-like access to FRenderPageManager. Will error if this module hasn't started (or has stopped).
+		 *
+		 * @return Returns RenderPageManager singleton instance.
+		 */
+		virtual FRenderPageManager& GetManager() const = 0;
+
+		/** Creates a URenderPagePropsSourceBase instance, based on the given ERenderPagePropsSourceType. */
+		virtual URenderPagePropsSourceBase* CreatePropsSource(UObject* Outer, ERenderPagePropsSourceType PropsSourceType, UObject* PropsSourceOrigin) = 0;
+		
+		/** Returns all set IRenderPagePropsSourceFactory instances, these are used to create URenderPagePropsSourceBase instances based on a given ERenderPagePropsSourceType. */
+		virtual const TMap<ERenderPagePropsSourceType, TSharedPtr<IRenderPagePropsSourceFactory>>& GetPropsSourceFactories() const = 0;
+	};
+}
diff --git a/Engine/Plugins/Experimental/RenderPages/Source/RenderPages/Public/RenderPage/RenderPageCollection.h b/Engine/Plugins/Experimental/RenderPages/Source/RenderPages/Public/RenderPage/RenderPageCollection.h
new file mode 100644
--- /dev/null
+++ b/Engine/Plugins/Experimental/RenderPages/Source/RenderPages/Public/RenderPage/RenderPageCollection.h
@@ -0,0 +1,310 @@
+// Copyright Epic Games, Inc. All Rights Reserved.
+
+#pragma once
+
+#include "CoreMinimal.h"
+#include "RenderPage/RenderPagePropsSource.h"
+#include "RenderPageCollection.generated.h"
+
+
+class UMoviePipelineOutputSetting;
+class UMoviePipelineMasterConfig;
+class ULevelSequence;
+class URenderPageCollection;
+
+
+/**
+ * This struct contains the data for a certain remote control property.
+ * 
+ * It's currently simply a wrapper around a byte array.
+ * This struct is needed so that that byte array can be used in another UPROPERTY container (TMap, TArray, etc).
+ */
+USTRUCT(BlueprintType)
+struct RENDERPAGES_API FRenderPageRemoteControlPropertyData
+{
+	GENERATED_BODY()
+
+public:
+	FRenderPageRemoteControlPropertyData() = default;
+	FRenderPageRemoteControlPropertyData(const TArray<uint8>& InBytes)
+		: Bytes(InBytes)
+	{}
+
+public:
+	/** The property data, as bytes. */
+	UPROPERTY()
+	TArray<uint8> Bytes;
+};
+
+
+/**
+ * This class represents a render page.
+ * It contains a level sequence and custom properties that will be applied while rendering.
+ * 
+ * Each render page must belong to a render page collection.
+ */
+UCLASS(BlueprintType)
+class RENDERPAGES_API URenderPage : public UObject
+{
+	GENERATED_BODY()
+
+public:
+	URenderPage();
+
+	/** Gets the calculated start frame, not taking the framerate of the render preset into account. */
+	TOptional<int32> GetSequenceStartFrame() const;
+
+	/** Gets the calculated end frame, not taking the framerate of the render preset into account. */
+	TOptional<int32> GetSequenceEndFrame() const;
+
+	/** Sets the custom start frame to match the given sequence start frame. */
+	bool SetSequenceStartFrame(const int32 NewCustomStartFrame);
+
+	/** Sets the custom end frame to match the given sequence end frame. */
+	bool SetSequenceEndFrame(const int32 NewCustomEndFrame);
+
+	/** Gets the calculated start frame. */
+	TOptional<int32> GetStartFrame() const;
+
+	/** Gets the calculated end frame. */
+	TOptional<int32> GetEndFrame() const;
+
+	/** Gets the calculated start time. */
+	TOptional<double> GetStartTime() const;
+
+	/** Gets the calculated end time. */
+	TOptional<double> GetEndTime() const;
+
+	/** Gets the calculated duration in seconds. */
+	TOptional<double> GetDurationInSeconds() const;
+
+	/** Gets the aspect ratio that this page will be rendered in. */
+	double GetOutputAspectRatio() const;
+
+	/** Checks whether the page contains data that matches the search terms. */
+	bool MatchesSearchTerm(const FString& SearchTerm) const;
+
+public:
+	FGuid GetId() const { return Id; }
+	void GenerateNewId() { Id = FGuid::NewGuid(); }
+
+	ULevelSequence* GetSequence() const { return Sequence; }
+	void SetSequence(ULevelSequence* NewSequence) { Sequence = NewSequence; }
+
+	bool GetIsCustomStartFrame() const { return bOverrideStartFrame; }
+	void SetIsCustomStartFrame(bool bNewOverrideStartFrame) { bOverrideStartFrame = bNewOverrideStartFrame; }
+	int32 GetCustomStartFrame() const { return CustomStartFrame; }
+	void SetCustomStartFrame(const int32 NewCustomStartFrame) { CustomStartFrame = NewCustomStartFrame; }
+
+	bool GetIsCustomEndFrame() const { return bOverrideEndFrame; }
+	void SetIsCustomEndFrame(bool bNewOverrideEndFrame) { bOverrideEndFrame = bNewOverrideEndFrame; }
+	int32 GetCustomEndFrame() const { return CustomEndFrame; }
+	void SetCustomEndFrame(const int32 NewCustomEndFrame) { CustomEndFrame = NewCustomEndFrame; }
+
+	bool GetIsCustomResolution() const { return bOverrideResolution; }
+	void SetIsCustomResolution(bool bNewOverrideResolution) { bOverrideResolution = bNewOverrideResolution; }
+	FIntPoint GetCustomResolution() const { return CustomResolution; }
+	void SetCustomResolution(const FIntPoint NewCustomResolution) { CustomResolution = NewCustomResolution; }
+
+	FString GetPageId() const { return PageId; }
+	static FString PurgePageIdOrReturnEmptyString(const FString& NewPageId);
+	static FString PurgePageId(const FString& NewPageId);
+	static FString PurgePageIdOrGenerateUniqueId(URenderPageCollection* PageCollection, const FString& NewPageId);
+	void SetPageId(const FString& NewPageId) { PageId = PurgePageId(NewPageId); }
+	void SetPageIdRaw(const FString& NewPageId) { PageId = NewPageId; }
+
+	FString GetPageName() const { return PageName; }
+	static FString PurgePageName(const FString& NewPageName);
+	void SetPageName(const FString& NewPageName) { PageName = NewPageName; }
+	void SetPageNameRaw(const FString& NewPageName) { PageName = NewPageName; }
+
+	bool GetIsEnabled() const { return bIsEnabled; }
+	void SetIsEnabled(const bool bEnabled) { bIsEnabled = bEnabled; }
+
+	FString GetOutputDirectory() const;
+	FString GetOutputDirectoryRaw() const { return OutputDirectory; }
+	FString GetOutputDirectoryForDisplay() const { return OutputDirectory; }
+	static FString PurgeOutputDirectory(const FString& NewOutputDirectory);
+	void SetOutputDirectory(const FString& NewOutputDirectory) { OutputDirectory = PurgeOutputDirectory(NewOutputDirectory); }
+	void SetOutputDirectoryRaw(const FString& NewOutputDirectory) { OutputDirectory = NewOutputDirectory; }
+
+	UMoviePipelineMasterConfig* GetRenderPreset() const { return RenderPreset; }
+	UMoviePipelineOutputSetting* GetRenderPresetOutputSettings() const;
+	void SetRenderPreset(UMoviePipelineMasterConfig* NewRenderPreset) { RenderPreset = NewRenderPreset; }
+
+	bool HasRemoteControlValue(const TSharedPtr<FRemoteControlEntity>& RemoteControlEntity) const;
+	bool ConstGetRemoteControlValue(const TSharedPtr<FRemoteControlEntity>& RemoteControlEntity, TArray<uint8>& OutBinaryArray) const;
+	bool GetRemoteControlValue(const TSharedPtr<FRemoteControlEntity>& RemoteControlEntity, TArray<uint8>& OutBinaryArray);
+	bool SetRemoteControlValue(const TSharedPtr<FRemoteControlEntity>& RemoteControlEntity, const TArray<uint8>& BinaryArray);
+	TMap<FString, FRenderPageRemoteControlPropertyData>& GetRemoteControlValuesRef() { return RemoteControlValues; }
+
+private:
+	/** The unique ID of this render page. */
+	UPROPERTY()
+	FGuid Id;
+
+	/** The level sequence, this is what will be rendered during rendering. A render page without a level sequence can't be rendered. */
+	UPROPERTY(EditInstanceOnly, Category="Render Pages|Page", Meta=(AllowPrivateAccess="true"))
+	TObjectPtr<ULevelSequence> Sequence;
+
+	/** If this is true, the CustomStartFrame property will override the start frame of the level sequence. */
+	UPROPERTY(EditInstanceOnly, Category="Render Pages|Page", Meta=(AllowPrivateAccess="true", InlineEditConditionToggle))
+	bool bOverrideStartFrame;
+
+	/** If bOverrideStartFrame is true, this property will override the start frame of the level sequence. */
+	UPROPERTY(EditInstanceOnly, Category="Render Pages|Page", Meta=(AllowPrivateAccess="true", EditCondition="bOverrideStartFrame"))
+	int32 CustomStartFrame;
+
+	/** If this is true, the CustomEndFrame property will override the end frame of the level sequence. */
+	UPROPERTY(EditInstanceOnly, Category="Render Pages|Page", Meta=(AllowPrivateAccess="true", InlineEditConditionToggle))
+	bool bOverrideEndFrame;
+
+	/** If bOverrideEndFrame is true, this property will override the end frame of the level sequence. */
+	UPROPERTY(EditInstanceOnly, Category="Render Pages|Page", Meta=(AllowPrivateAccess="true", EditCondition="bOverrideEndFrame"))
+	int32 CustomEndFrame;
+
+	/** If this is true, the CustomResolution property will override the resolution of the render. */
+	UPROPERTY()
+	bool bOverrideResolution;
+
+	/** If bOverrideResolution is true, this property will override the resolution of the render. */
+	UPROPERTY()
+	FIntPoint CustomResolution;
+
+	/** If this is true, this render page will be rendered during a batch rendering, otherwise it will be skipped. */
+	UPROPERTY()
+	bool bIsEnabled;
+
+	/** The 'ID' of this page, this 'ID' is set by users. During rendering it will place all the output files of this render page into a folder called after this 'ID', this means that this string can only contain file-safe characters. */
+	UPROPERTY()
+	FString PageId;
+
+	/** The name of this page, this can be anything, it's set by the user, it serves as a way for the user to understand what page this is. */
+	UPROPERTY()
+	FString PageName;
+
+	/** This is the folder in which the output files (of rendering) are placed into. To be more specific, the output files are placed in: {OutputDirectory}/{PageId}/, this folder will be created if it doesn't exist at the time of rendering. */
+	UPROPERTY()
+	FString OutputDirectory;
+
+	/** The MRQ render preset. The pages are rendered using the MRQ (Movie Render Queue) plugin. This 'preset' contains the configuration of that plugin. */
+	UPROPERTY()
+	TObjectPtr<UMoviePipelineMasterConfig> RenderPreset;
+
+	/** The Remote Control plugin can be used to customize and modify the way a page is rendered. If Remote Control is being used, the property values of this page will be stored in this map (remote control entity id -> value as bytes). */
+	UPROPERTY()
+	TMap<FString, FRenderPageRemoteControlPropertyData> RemoteControlValues;
+};
+
+
+/**
+ * This class represents a collection of render pages.
+ * A render page collection is the asset that is shown in the content browser, it's the asset that can be opened and edited using the editor.
+ */
+UCLASS(Blueprintable, BlueprintType, EditInlineNew, Meta=(DontUseGenericSpawnObject="true"))
+class RENDERPAGES_API URenderPageCollection : public UObject
+{
+	GENERATED_BODY()
+
+public:
+	URenderPageCollection();
+
+	//UObject interface
+	virtual UWorld* GetWorld() const override;
+	virtual void PreSave(FObjectPreSaveContext SaveContext) override;
+	virtual void PostLoad() override;
+	//~ End UObject Interface
+
+	/** Should be called when the editor closes this asset. */
+	void OnClose() { SaveValuesToCDO(); }
+
+private:
+	/**
+	 * Because render page collection assets are blueprints (assets that also have a blueprint graph), the render page collection data is not stored directly in the asset data that you see in the content browser.
+	 * Instead, the data that is stored (and load) is the CDO (class default object).
+	 * Because of that, any data that needs to persist needs to be copied over to the CDO during a save, and data you'd like to load from it needs to be copied from the CDO during a load.
+	 */
+
+	/** Obtains the CDO, could return itself if this is called on the CDO instance. */
+	URenderPageCollection* GetCDO() { return (HasAnyFlags(RF_ClassDefaultObject) ? this : GetClass()->GetDefaultObject<URenderPageCollection>()); }
+	/** Copied values over into the CDO. */
+	void SaveValuesToCDO() { CopyValuesToOrFromCDO(true); }
+	/** Copied values over from the CDO. */
+	void LoadValuesFromCDO() { CopyValuesToOrFromCDO(false); }
+	/** Copied values to or from the CDO, based on whether bToCDO is true or false. */
+	void CopyValuesToOrFromCDO(const bool bToCDO);
+
+public:
+	FGuid GetId() const { return Id; }
+	void GenerateNewId() { Id = FGuid::NewGuid(); }
+
+	void SetPropsSource(ERenderPagePropsSourceType InPropsSourceType, UObject* InPropsSourceOrigin = nullptr);
+	URenderPagePropsSourceBase* GetPropsSource() const;
+
+	template<typename Type>
+	Type* GetPropsSource() const
+	{
+		static_assert(TIsDerivedFrom<Type, URenderPagePropsSourceBase>::IsDerived, "Type needs to derive from URenderPagePropsSourceBase.");
+		return Cast<Type>(GetPropsSource());
+	}
+
+	ERenderPagePropsSourceType GetPropsSourceType() const { return PropsSourceType; }
+	UObject* GetPropsSourceOrigin() const;
+
+public:
+	void AddRenderPage(URenderPage* RenderPage);
+	void RemoveRenderPage(URenderPage* RenderPage);
+	void InsertRenderPage(URenderPage* RenderPage, int32 Index);
+	bool HasRenderPage(URenderPage* RenderPage) const;
+	int32 GetIndexOfRenderPage(URenderPage* RenderPage) const;
+	TArray<TObjectPtr<URenderPage>>& GetRenderPagesRef() { return RenderPages; }
+	TArray<URenderPage*> GetRenderPages() const;
+	TArray<URenderPage*> GetEnabledRenderPages() const;
+	TArray<URenderPage*> GetDisabledRenderPages() const;
+	void InsertRenderPageBefore(URenderPage* RenderPage, URenderPage* BeforeRenderPage);
+	void InsertRenderPageAfter(URenderPage* RenderPage, URenderPage* AfterRenderPage);
+
+private:
+	DECLARE_MULTICAST_DELEGATE(FOnRenderPageCollectionPreSave);
+	FOnRenderPageCollectionPreSave& OnPreSaveCDO() { return GetCDO()->OnRenderPageCollectionPreSaveDelegate; }
+
+private:
+	/** The delegate for when this collection is about to save. */
+	FOnRenderPageCollectionPreSave OnRenderPageCollectionPreSaveDelegate;
+
+private:
+	/** The unique ID of this render page collection. */
+	UPROPERTY()
+	FGuid Id;
+
+
+	/** The type of the properties that a page can have. */
+	UPROPERTY(EditInstanceOnly, Category="Render Pages|Page Collection", Meta=(DisplayName="Properties Type", AllowPrivateAccess="true"))
+	ERenderPagePropsSourceType PropsSourceType;
+
+	/** The remote control properties that a page can have, only use this if PropsSourceType is ERenderPagePropsSourceType::RemoteControl. */
+	UPROPERTY(EditInstanceOnly, Category="Render Pages|Page Collection", Meta=(DisplayName="Remote Control Preset", AllowPrivateAccess="true", EditCondition="PropsSourceType == ERenderPagePropsSourceType::RemoteControl", EditConditionHides))
+	TObjectPtr<URemoteControlPreset> PropsSourceOrigin_RemoteControl;
+
+
+	/** The render pages of this collection. */
+	UPROPERTY(Instanced)
+	TArray<TObjectPtr<URenderPage>> RenderPages;
+
+
+	/** GetPropsSource calls are somewhat expensive, we speed that up by caching the result (the PropsSource) that has been last outputted by that function. */
+	UPROPERTY(Transient)
+	mutable TObjectPtr<URenderPagePropsSourceBase> CachedPropsSource;
+
+	/** GetPropsSource calls are somewhat expensive, we speed that up by caching the PropsSourceType last used in that function. */
+	UPROPERTY(Transient)
+	mutable ERenderPagePropsSourceType CachedPropsSourceType;
+	
+	/** GetPropsSource calls are somewhat expensive, we speed that up by caching the PropsSourceOrigin last used in that function. */
+	UPROPERTY(Transient)
+	mutable TWeakObjectPtr<UObject> CachedPropsSourceOriginWeakPtr;
+
+
+	/** GetWorld calls can be expensive, we speed them up by caching the last found world until it goes away. */
+	mutable TWeakObjectPtr<UWorld> CachedWorldWeakPtr;
+};
diff --git a/Engine/Plugins/Experimental/RenderPages/Source/RenderPages/Public/RenderPage/RenderPageManager.h b/Engine/Plugins/Experimental/RenderPages/Source/RenderPages/Public/RenderPage/RenderPageManager.h
new file mode 100644
--- /dev/null
+++ b/Engine/Plugins/Experimental/RenderPages/Source/RenderPages/Public/RenderPage/RenderPageManager.h
@@ -0,0 +1,137 @@
+// Copyright Epic Games, Inc. All Rights Reserved.
+
+#pragma once
+
+#include "CoreMinimal.h"
+#include "RenderPage/RenderPageCollection.h"
+#include "RenderPageManager.generated.h"
+
+
+class URenderPage;
+class URenderPageCollection;
+class URenderPagePropRemoteControl;
+class URenderPagesMoviePipelineRenderJob;
+class UTexture2D;
+
+
+/**
+ * This struct keeps track of the values of the properties before new ones were applied, so we can rollback to the previous state.
+ */
+USTRUCT()
+struct RENDERPAGES_API FRenderPageManagerPreviousPagePropValues
+{
+	GENERATED_BODY()
+
+public:
+	FRenderPageManagerPreviousPagePropValues() = default;
+	FRenderPageManagerPreviousPagePropValues(const TMap<TObjectPtr<URenderPagePropRemoteControl>, FRenderPageRemoteControlPropertyData>& InRemoteControlData)
+		: RemoteControlData(InRemoteControlData)
+	{}
+
+public:
+	/** The previous values of the remote control properties. */
+	UPROPERTY()
+	TMap<TObjectPtr<URenderPagePropRemoteControl>, FRenderPageRemoteControlPropertyData> RemoteControlData;
+};
+
+
+namespace UE::RenderPages
+{
+	/** A delegate for when FRenderPageManager::RenderPreviewFrame has finished. */
+	DECLARE_DELEGATE_OneParam(FRenderPageManagerRenderPreviewFrameArgsCallback, bool /*bSuccess*/);
+
+
+	/**
+	 * The arguments for the FRenderPageManager::RenderPreviewFrame function.
+	 */
+	struct RENDERPAGES_API FRenderPageManagerRenderPreviewFrameArgs
+	{
+	public:
+		/** Whether it should run invisibly (so without any UI elements popping up during rendering) or not. */
+		bool bHeadless = false;
+
+		/** The render page collection of the given render pages that will be rendered. */
+		TObjectPtr<URenderPageCollection> PageCollection = nullptr;
+
+		/** The specific render page that will be rendered. */
+		TObjectPtr<URenderPage> Page = nullptr;
+
+		/** The specific frame number that will be rendered. */
+		TOptional<int32> Frame;
+
+		/** The resolution it will be rendered in. */
+		FIntPoint Resolution = FIntPoint(0, 0);
+
+		/** The texture to reuse for rendering (performance optimization, prevents a new UTexture2D from having to be created, will only be used if the resolution of this texture matches the resolution it will be rendering in). */
+		TObjectPtr<UTexture2D> ReusingTexture2D = nullptr;
+
+		/** The delegate for when the rendering has finished. */
+		FRenderPageManagerRenderPreviewFrameArgsCallback Callback;
+	};
+
+
+	/**
+	 * The singleton class that manages the render pages.
+	 * 
+	 * This functionality is separated from the UI in order to make it is reusable, meaning that it can also be used in other modules.
+	 */
+	class RENDERPAGES_API FRenderPageManager
+	{
+	public:
+		/** A folder in which rendered frames for temporary use will be placed in. */
+		inline static FString TmpRenderedFramesPath = FPaths::AutomationTransientDir() / TEXT("RenderPages");
+
+		/** The number of characters for a generated ID. For example, a value of 4 results in IDs: "0001", "0002", etc. */
+		static constexpr int32 GeneratedIdCharacterLength = 4;
+
+	public:
+		/** Creates a new page and adds it to the given collection. **/
+		URenderPage* AddNewPage(URenderPageCollection* PageCollection);
+
+		/** Copy the given pages in the given collection. **/
+		URenderPage* CopyPage(URenderPageCollection* PageCollection, URenderPage* Page);
+
+		/** Delete the given page in the given collection. **/
+		void DeletePage(URenderPageCollection* PageCollection, URenderPage* Page);
+
+		/** Relocates the given page in the given collection to the position of the given dropped-on page. **/
+		bool DragDropPage(URenderPageCollection* PageCollection, URenderPage* Page, URenderPage* DroppedOnPage, const bool bAfter = true);
+
+		/** Generates a unique page ID by finding the highest page ID and increasing it by one. **/
+		FString CreateUniquePageId(URenderPageCollection* PageCollection);
+
+		/** Finds whether given page ID already exists in the collection. **/
+		bool DoesPageIdExist(URenderPageCollection* PageCollection, const FString& PageId);
+
+		/** Batch render the current pages of the given collection. **/
+		URenderPagesMoviePipelineRenderJob* CreateBatchRenderJob(URenderPageCollection* PageCollection);
+
+		/** Batch render the current pages of the given collection. **/
+		URenderPagesMoviePipelineRenderJob* RenderPreviewFrame(const FRenderPageManagerRenderPreviewFrameArgs& Args);
+
+		/** Gets the rendered preview frame (of a rendering in which the frame number was specified). */
+		UTexture2D* GetSingleRenderedPreviewFrame(URenderPage* Page, UTexture2D* ReusingTexture2D, bool& bOutReusedGivenTexture2D);
+
+		/** Gets the rendered preview frame (of a rendering in which the frame number was specified). */
+		UTexture2D* GetSingleRenderedPreviewFrame(URenderPage* Page);
+
+		/** Gets the rendered preview frame the given frame number (of a rendering in which the frame number was not specified). */
+		UTexture2D* GetRenderedPreviewFrame(URenderPage* Page, const int32 Frame, UTexture2D* ReusingTexture2D, bool& bOutReusedGivenTexture2D);
+
+		/** Gets the rendered preview frame the given frame number (of a rendering in which the frame number was not specified). */
+		UTexture2D* GetRenderedPreviewFrame(URenderPage* Page, const int32 Frame);
+
+		/** Makes sure that all the data from the current props source is stored in all of the pages of this page collection. */
+		void UpdatePagesPropValues(URenderPageCollection* PageCollection);
+
+		/** Applies the props of the given page, also requires the page collection to be given as well (to know what props the page is using). */
+		FRenderPageManagerPreviousPagePropValues ApplyPagePropValues(const URenderPageCollection* PageCollection, const URenderPage* Page);
+
+		/** Restores the props that were previously applied, to the values they were before. */
+		void RestorePagePropValues(const FRenderPageManagerPreviousPagePropValues& PreviousPropValues);
+
+	private:
+		/** The map that stores the start frame (of a render) of each rendered page. */
+		TMap<FGuid, int32> StartFrameOfRenders;
+	};
+}
diff --git a/Engine/Plugins/Experimental/RenderPages/Source/RenderPages/Public/RenderPage/RenderPageMoviePipelineJob.h b/Engine/Plugins/Experimental/RenderPages/Source/RenderPages/Public/RenderPage/RenderPageMoviePipelineJob.h
new file mode 100644
--- /dev/null
+++ b/Engine/Plugins/Experimental/RenderPages/Source/RenderPages/Public/RenderPage/RenderPageMoviePipelineJob.h
@@ -0,0 +1,128 @@
+// Copyright Epic Games, Inc. All Rights Reserved.
+
+#pragma once
+
+#include "RenderPageManager.h"
+#include "RenderPageMoviePipelineJob.generated.h"
+
+
+class UMoviePipelineSetting;
+class UMoviePipelineOutputBase;
+struct FMoviePipelineOutputData;
+class URenderPageCollection;
+class URenderPage;
+class UMoviePipelineExecutorBase;
+class UMoviePipelineQueue;
+class UMoviePipelineExecutorJob;
+class URenderPagesMoviePipelineRenderJob;
+
+
+namespace UE::RenderPages
+{
+	/** A delegate for when a render job is about to start. */
+	DECLARE_MULTICAST_DELEGATE_OneParam(FOnRenderPagesMoviePipelineRenderJobStarted, URenderPagesMoviePipelineRenderJob* /*RenderJob*/);
+
+	/** A delegate for when a render job has finished. */
+	DECLARE_MULTICAST_DELEGATE_TwoParams(FOnRenderPagesMoviePipelineRenderJobFinished, URenderPagesMoviePipelineRenderJob* /*RenderJob*/, bool /*bSuccess*/);
+
+	/**
+	 * The arguments for the URenderPagesMoviePipelineRenderJob::Create function.
+	 */
+	struct RENDERPAGES_API FRenderPagesMoviePipelineRenderJobCreateArgs
+	{
+	public:
+		/** The render page collection of the given render pages that will be rendered. */
+		TObjectPtr<URenderPageCollection> PageCollection = nullptr;
+
+		/** The specific render pages that will be rendered. */
+		TArray<TObjectPtr<URenderPage>> Pages;
+
+		/** If not null, it will override the MRQ pipeline executor class with this class. */
+		TSubclassOf<UMoviePipelineExecutorBase> PipelineExecutorClass = nullptr;
+
+		/** The MRQ settings classes to disable (things like Anti-Aliasing, High-Res, etc). */
+		TArray<TSubclassOf<UMoviePipelineSetting>> DisableSettingsClasses;
+
+		/** Whether it should run invisibly (so without any UI elements popping up during rendering) or not. */
+		bool bHeadless = false;
+
+		/** Whether it should make sure it will output an image or not (if this bool is true, it will test if JPG/PNG/etc output is enabled, if none are, it will enable PNG output). */
+		bool bForceOutputImage = false;
+
+		/** Whether it should make sure it will only output in a single format (if this bool is true, if for example JPG and PNG output are enabled, one will be disabled, so that there will only be 1 output that's enabled). */
+		bool bForceOnlySingleOutput = false;
+
+		/** Whether it should use the sequence's framerate rather than any manually set framerate (if this bool is true, it will make sure bUseCustomFrameRate is set to false). */
+		bool bForceUseSequenceFrameRate = false;
+
+		/** Whether it should make sure it will output files named 0000000001, 0000000002, etc (if this bool is true, it will override the FileNameFormat to simply output the frame number, and it will add 1000000000 to that frame number to hopefully ensure that any negative frame numbers will not result in filenames starting with a minus character). */
+		bool bEnsureSequentialFilenames = false;
+	};
+}
+
+
+/**
+ * This class is responsible for rendering render pages.
+ */
+UCLASS()
+class RENDERPAGES_API URenderPagesMoviePipelineRenderJob : public UObject
+{
+	GENERATED_BODY()
+
+public:
+	/** Creates a new render job instance, it won't be started right away. */
+	static URenderPagesMoviePipelineRenderJob* Create(const UE::RenderPages::FRenderPagesMoviePipelineRenderJobCreateArgs& Args);
+
+	/** Starts this render job. */
+	void Execute();
+
+	/** Cancels this render job. Relies on the internal MRQ implementation of job canceling on whether this will do anything or not. */
+	void Cancel();
+
+	/** Retrieves the rendering status of the given render page. */
+	FString GetPageStatus(URenderPage* Page) const;
+
+private:
+	void ComputePlaybackContext(bool& bOutAllowBinding);
+	void ExecutePageStarted(UMoviePipelineExecutorJob* JobToStart);
+	void ExecutePageFinished(FMoviePipelineOutputData PipelineOutputData);
+	void ExecuteFinished(UMoviePipelineExecutorBase* PipelineExecutor, const bool bSuccess);
+
+
+protected:
+	/** The MRQ queue. */
+	UPROPERTY(Transient)
+	TObjectPtr<UMoviePipelineQueue> RenderQueue;
+
+	/** The MRQ pipeline executor. */
+	UPROPERTY(Transient)
+	TObjectPtr<UMoviePipelineExecutorBase> ActiveExecutor;
+
+	/** The render page collection of the given render pages that will be rendered. */
+	UPROPERTY(Transient)
+	TObjectPtr<URenderPageCollection> PageCollection;
+
+	/** The render pages that are to be rendered, mapped to the job of each specific render page. */
+	UPROPERTY(Transient)
+	TMap<TObjectPtr<const URenderPage>, TObjectPtr<UMoviePipelineExecutorJob>> PageExecutorJobs;
+
+	/** The render pages that are to be rendered, mapped to the rendering status of each specific render page. */
+	UPROPERTY(Transient)
+	TMap<TObjectPtr<const URenderPage>, FString> PageStatuses;
+
+	/** The render page properties that have been overwritten by the currently applied page properties. */
+	UPROPERTY(Transient)
+	FRenderPageManagerPreviousPagePropValues PreviousPageProps;
+
+
+public:
+	/** A delegate for when the render job is about to start. */
+	UE::RenderPages::FOnRenderPagesMoviePipelineRenderJobStarted& OnExecuteStarted() { return OnExecuteStartedDelegate; }
+
+	/** A delegate for when the render job has finished. */
+	UE::RenderPages::FOnRenderPagesMoviePipelineRenderJobFinished& OnExecuteFinished() { return OnExecuteFinishedDelegate; }
+
+private:
+	UE::RenderPages::FOnRenderPagesMoviePipelineRenderJobStarted OnExecuteStartedDelegate;
+	UE::RenderPages::FOnRenderPagesMoviePipelineRenderJobFinished OnExecuteFinishedDelegate;
+};
