First usable version of the render pages plugin.

diff --git a/Engine/Plugins/Experimental/RenderPages/Source/RenderPages/Public/RenderPage/RenderPagePropsSource.h b/Engine/Plugins/Experimental/RenderPages/Source/RenderPages/Public/RenderPage/RenderPagePropsSource.h
new file mode 100644
--- /dev/null
+++ b/Engine/Plugins/Experimental/RenderPages/Source/RenderPages/Public/RenderPage/RenderPagePropsSource.h
@@ -0,0 +1,208 @@
+// Copyright Epic Games, Inc. All Rights Reserved.
+
+#pragma once
+
+#include "CoreMinimal.h"
+#include "RemoteControlPreset.h"
+#include "RenderPagePropsSource.generated.h"
+
+
+/////////////////////////////////////////////////////
+// Enums
+
+/**
+ * The type of the properties source.
+ * In other words, where the properties come from that each render page can have.
+ */
+UENUM(BlueprintType)
+enum class ERenderPagePropsSourceType : uint8
+{
+	Local = 0 UMETA(DisplayName = "Local Source"),
+	RemoteControl = 1 UMETA(DisplayName = "Remote Control Preset")
+};
+
+
+/////////////////////////////////////////////////////
+// Base classes
+
+/**
+ * The base class of the render page property abstraction.
+ */
+UCLASS(Abstract, HideCategories=Object)
+class RENDERPAGES_API URenderPagePropBase : public UObject
+{
+	GENERATED_BODY()
+};
+
+/**
+ * The base class of the render page properties abstraction.
+ */
+UCLASS(Abstract, HideCategories=Object)
+class RENDERPAGES_API URenderPagePropsBase : public UObject
+{
+	GENERATED_BODY()
+
+public:
+	/** Returns all props. */
+	virtual TArray<URenderPagePropBase*> GetAll() const { return TArray<URenderPagePropBase*>(); }
+};
+
+/**
+ * The base class of the render page properties source abstraction.
+ */
+UCLASS(Abstract, BlueprintType, HideCategories=Object)
+class RENDERPAGES_API URenderPagePropsSourceBase : public UObject
+{
+	GENERATED_BODY()
+
+public:
+	/** Returns the ID, which is randomly generated when an instance of this class is constructed. */
+	FGuid GetId() const { return Id; }
+
+	/** Randomly generates a new ID. */
+	void GenerateNewId() { Id = FGuid::NewGuid(); }
+
+
+	/** Returns the type of this properties source. */
+	virtual ERenderPagePropsSourceType GetType() const { return ERenderPagePropsSourceType::Local; }
+
+	/** Sets the properties source. */
+	virtual void SetSourceOrigin(UObject* SourceOrigin) {}
+
+	/** Returns the collection of properties (that this properties source contains). */
+	virtual URenderPagePropsBase* GetProps() const { return NewObject<URenderPagePropsBase>(const_cast<URenderPagePropsSourceBase*>(this)); }
+
+protected:
+	UPROPERTY()
+	FGuid Id = FGuid::NewGuid();
+};
+
+
+/////////////////////////////////////////////////////
+// Local source classes
+
+/**
+ * The local properties implementation of the render page property abstraction.
+ */
+UCLASS(HideCategories=Object)
+class RENDERPAGES_API URenderPagePropLocal : public URenderPagePropBase
+{
+	GENERATED_BODY()
+};
+
+/**
+ * The local properties implementation of the render page properties abstraction.
+ */
+UCLASS(HideCategories=Object)
+class RENDERPAGES_API URenderPagePropsLocal : public URenderPagePropsBase
+{
+	GENERATED_BODY()
+};
+
+/**
+ * The local properties implementation of the render page properties source abstraction.
+ */
+UCLASS(BlueprintType, HideCategories=Object)
+class RENDERPAGES_API URenderPagePropsSourceLocal : public URenderPagePropsSourceBase
+{
+	GENERATED_BODY()
+
+public:
+	//~ Begin URenderPagePropsSourceBase interface
+	virtual ERenderPagePropsSourceType GetType() const override { return ERenderPagePropsSourceType::Local; }
+	virtual URenderPagePropsBase* GetProps() const override { return NewObject<URenderPagePropsLocal>(const_cast<URenderPagePropsSourceLocal*>(this)); }
+	//~ End URenderPagePropsSourceBase interface
+};
+
+
+/////////////////////////////////////////////////////
+// Remote control source classes
+
+/**
+ * The remote control properties implementation of the render page property abstraction.
+ */
+UCLASS(HideCategories=Object)
+class RENDERPAGES_API URenderPagePropRemoteControl : public URenderPagePropBase
+{
+	GENERATED_BODY()
+
+public:
+	/** Gets the value (as bytes) of the given property (remote control entity). Returns true if the operation was successful, false otherwise. */
+	static bool GetValueOfEntity(const TSharedPtr<FRemoteControlEntity>& RemoteControlEntity, TArray<uint8>& OutBinaryArray);
+
+	/** Set the value (as bytes) of the given property (remote control entity). Returns true if the operation was successful, false otherwise. */
+	static bool SetValueOfEntity(const TSharedPtr<FRemoteControlEntity>& RemoteControlEntity, const TArray<uint8>& BinaryArray);
+
+public:
+	/** Sets the initial values of this instance. */
+	void Initialize(const TSharedPtr<FRemoteControlEntity>& InRemoteControlEntity);
+
+	/** Returns the property, which is a remote control entity (which can be a field or a function). */
+	TSharedPtr<FRemoteControlEntity> GetRemoteControlEntity() const { return RemoteControlEntity; }
+
+	/** Gets the value (as bytes) of this property. Returns true if the operation was successful, false otherwise. */
+	bool GetValue(TArray<uint8>& OutBinaryArray) const { return GetValueOfEntity(RemoteControlEntity, OutBinaryArray); }
+
+	/** Set the value (as bytes) of this property. Returns true if the operation was successful, false otherwise. */
+	bool SetValue(const TArray<uint8>& BinaryArray) { return SetValueOfEntity(RemoteControlEntity, BinaryArray); }
+
+protected:
+	/** The property, which is a remote control entity (which can be a field or a function). */
+	TSharedPtr<FRemoteControlEntity> RemoteControlEntity = nullptr;
+};
+
+/**
+ * The remote control properties implementation of the render page properties abstraction.
+ */
+UCLASS(HideCategories=Object)
+class RENDERPAGES_API URenderPagePropsRemoteControl : public URenderPagePropsBase
+{
+	GENERATED_BODY()
+
+public:
+	/** Sets the initial values of this instance. */
+	void Initialize(URemoteControlPreset* InRemoteControlPreset);
+
+	//~ Begin URenderPagePropsBase interface
+	virtual TArray<URenderPagePropBase*> GetAll() const override;
+	//~ End URenderPagePropsBase interface
+
+	/** Returns all props, casted to URenderPagePropRemoteControl, for ease of use. */
+	TArray<URenderPagePropRemoteControl*> GetAllCasted() const;
+
+	/** Returns the remote control preset. */
+	URemoteControlPreset* GetRemoteControlPreset() const { return (IsValid(RemoteControlPreset) ? RemoteControlPreset : nullptr); }
+
+protected:
+	/** The source of properties, which is a remote control preset. */
+	UPROPERTY()
+	TObjectPtr<URemoteControlPreset> RemoteControlPreset = nullptr;
+};
+
+/**
+ * The remote control properties implementation of the render page properties source abstraction.
+ */
+UCLASS(BlueprintType, HideCategories=Object)
+class RENDERPAGES_API URenderPagePropsSourceRemoteControl : public URenderPagePropsSourceBase
+{
+	GENERATED_BODY()
+
+public:
+	//~ Begin URenderPagePropsSourceBase interface
+	virtual ERenderPagePropsSourceType GetType() const override { return ERenderPagePropsSourceType::RemoteControl; }
+	virtual void SetSourceOrigin(UObject* SourceOrigin) override;
+	virtual URenderPagePropsRemoteControl* GetProps() const override;
+	//~ End URenderPagePropsSourceBase interface
+
+	/** Returns (in the out parameter) the preset groups that are available in this remote control preset. */
+	void GetAvailablePresetGroups(TArray<FName>& OutPresetGroups) const;
+
+protected:
+	/** The source of properties, which is a remote control preset. */
+	UPROPERTY(VisibleInstanceOnly, Category="Render Pages|Remote Control")
+	TObjectPtr<URemoteControlPreset> RemoteControlPreset = nullptr;
+
+	/** The preset group (of the remote control preset) that we should obtain the properties from. */
+	UPROPERTY(VisibleInstanceOnly, Category="Render Pages|Remote Control")
+	FName ActivePresetGroup;
+};
diff --git a/Engine/Plugins/Experimental/RenderPages/Source/RenderPages/Public/RenderPage/RenderPagesBlueprintGeneratedClass.h b/Engine/Plugins/Experimental/RenderPages/Source/RenderPages/Public/RenderPage/RenderPagesBlueprintGeneratedClass.h
new file mode 100644
--- /dev/null
+++ b/Engine/Plugins/Experimental/RenderPages/Source/RenderPages/Public/RenderPage/RenderPagesBlueprintGeneratedClass.h
@@ -0,0 +1,23 @@
+// Copyright Epic Games, Inc. All Rights Reserved.
+
+#pragma once
+
+#include "CoreMinimal.h"
+#include "Engine/BlueprintGeneratedClass.h"
+#include "RenderPagesBlueprintGeneratedClass.generated.h"
+
+
+/**
+ * A UBlueprintGeneratedClass child class for the RenderPages modules.
+ *
+ * Required in order for a RenderPageCollection to be able to have a blueprint graph.
+ */
+UCLASS()
+class RENDERPAGES_API URenderPagesBlueprintGeneratedClass : public UBlueprintGeneratedClass
+{
+	GENERATED_BODY()
+public:
+	//~ Begin UBlueprintGeneratedClass interface
+	virtual uint8* GetPersistentUberGraphFrame(UObject* Obj, UFunction* FuncToCheck) const override;
+	//~ End UBlueprintGeneratedClass interface
+};
diff --git a/Engine/Plugins/Experimental/RenderPages/Source/RenderPages/Public/RenderPagesLog.h b/Engine/Plugins/Experimental/RenderPages/Source/RenderPages/Public/RenderPagesLog.h
new file mode 100644
--- /dev/null
+++ b/Engine/Plugins/Experimental/RenderPages/Source/RenderPages/Public/RenderPagesLog.h
@@ -0,0 +1,8 @@
+// Copyright Epic Games, Inc. All Rights Reserved.
+
+#pragma once
+
+#include "CoreMinimal.h"
+
+
+DECLARE_LOG_CATEGORY_CLASS(LogRenderPages, Log, All);
diff --git a/Engine/Plugins/Experimental/RenderPages/Source/RenderPages/RenderPages.Build.cs b/Engine/Plugins/Experimental/RenderPages/Source/RenderPages/RenderPages.Build.cs
new file mode 100644
--- /dev/null
+++ b/Engine/Plugins/Experimental/RenderPages/Source/RenderPages/RenderPages.Build.cs
@@ -0,0 +1,34 @@
+// Copyright Epic Games, Inc. All Rights Reserved.
+
+using UnrealBuildTool;
+
+public class RenderPages : ModuleRules
+{
+	public RenderPages(ReadOnlyTargetRules Target) : base(Target)
+	{
+		PCHUsage = PCHUsageMode.UseExplicitOrSharedPCHs;
+
+		PublicDependencyModuleNames.AddRange(
+			new string[]
+			{
+				"Core",
+				"CoreUObject",
+				"Engine",
+				"RemoteControl",
+			}
+		);
+
+
+		PrivateDependencyModuleNames.AddRange(
+			new string[]
+			{
+				"LevelSequence",
+				"LevelSequenceEditor",
+				"MovieRenderPipelineCore",
+				"MovieRenderPipelineEditor",
+				"MovieRenderPipelineRenderPasses",
+				"Serialization",
+			}
+		);
+	}
+}
\ No newline at end of file
diff --git a/Engine/Plugins/Experimental/RenderPages/Source/RenderPagesDeveloper/Private/Blueprints/RenderPagesBlueprint.cpp b/Engine/Plugins/Experimental/RenderPages/Source/RenderPagesDeveloper/Private/Blueprints/RenderPagesBlueprint.cpp
new file mode 100644
--- /dev/null
+++ b/Engine/Plugins/Experimental/RenderPages/Source/RenderPagesDeveloper/Private/Blueprints/RenderPagesBlueprint.cpp
@@ -0,0 +1,47 @@
+// Copyright Epic Games, Inc. All Rights Reserved.
+
+#include "Blueprints/RenderPagesBlueprint.h"
+#include "Graph/RenderPagesGraph.h"
+#include "RenderPage/RenderPageCollection.h"
+#include "RenderPage/RenderPagesBlueprintGeneratedClass.h"
+
+
+URenderPagesBlueprint::URenderPagesBlueprint()
+{
+	CompileLog.SetSourcePath(GetPathName());
+	CompileLog.bLogDetailedResults = false;
+	CompileLog.EventDisplayThresholdMs = false;
+}
+
+UClass* URenderPagesBlueprint::GetBlueprintClass() const
+{
+	return URenderPagesBlueprintGeneratedClass::StaticClass();
+}
+
+void URenderPagesBlueprint::PostLoad()
+{
+	Super::PostLoad();
+
+	{
+		// Remove all non Render Pages graphs
+		TArray<UEdGraph*> NewUberGraphPages;
+
+		for (UEdGraph* Graph : UbergraphPages)
+		{
+			if (URenderPagesGraph* RenderPagesGraph = Cast<URenderPagesGraph>(Graph))
+			{
+				NewUberGraphPages.Add(RenderPagesGraph);
+			}
+			else
+			{
+				Graph->MarkAsGarbage();
+				Graph->Rename(nullptr, GetTransientPackage(), REN_ForceNoResetLoaders);
+			}
+		}
+
+		UbergraphPages = NewUberGraphPages;
+	}
+
+	CompileLog.Messages.Reset();
+	CompileLog.NumErrors = CompileLog.NumWarnings = 0;
+}
diff --git a/Engine/Plugins/Experimental/RenderPages/Source/RenderPagesDeveloper/Private/Blueprints/RenderPagesBlueprintCompiler.cpp b/Engine/Plugins/Experimental/RenderPages/Source/RenderPagesDeveloper/Private/Blueprints/RenderPagesBlueprintCompiler.cpp
new file mode 100644
--- /dev/null
+++ b/Engine/Plugins/Experimental/RenderPages/Source/RenderPagesDeveloper/Private/Blueprints/RenderPagesBlueprintCompiler.cpp
@@ -0,0 +1,11 @@
+// Copyright Epic Games, Inc. All Rights Reserved.
+
+#include "Blueprints/RenderPagesBlueprintCompiler.h"
+
+
+bool UE::RenderPages::FRenderPagesBlueprintCompiler::CanCompile(const UBlueprint* Blueprint)
+{
+	return false;
+}
+
+void UE::RenderPages::FRenderPagesBlueprintCompiler::Compile(UBlueprint* Blueprint, const FKismetCompilerOptions& CompileOptions, FCompilerResultsLog& Results) {}
diff --git a/Engine/Plugins/Experimental/RenderPages/Source/RenderPagesDeveloper/Private/Graph/RenderPagesGraph.cpp b/Engine/Plugins/Experimental/RenderPages/Source/RenderPagesDeveloper/Private/Graph/RenderPagesGraph.cpp
new file mode 100644
--- /dev/null
+++ b/Engine/Plugins/Experimental/RenderPages/Source/RenderPagesDeveloper/Private/Graph/RenderPagesGraph.cpp
@@ -0,0 +1,29 @@
+// Copyright Epic Games, Inc. All Rights Reserved.
+
+#include "Graph/RenderPagesGraph.h"
+#include "Graph/RenderPagesGraphSchema.h"
+#include "Blueprints/RenderPagesBlueprint.h"
+#include "Kismet2/BlueprintEditorUtils.h"
+
+#define LOCTEXT_NAMESPACE "RenderPagesGraph"
+
+
+void URenderPagesGraph::Initialize(URenderPagesBlueprint* InBlueprint) {}
+
+void URenderPagesGraph::PostLoad()
+{
+	Super::PostLoad();
+	Schema = URenderPagesGraphSchema::StaticClass();
+}
+
+URenderPagesBlueprint* URenderPagesGraph::GetBlueprint() const
+{
+	if (URenderPagesGraph* OuterGraph = Cast<URenderPagesGraph>(GetOuter()))
+	{
+		return OuterGraph->GetBlueprint();
+	}
+	return Cast<URenderPagesBlueprint>(GetOuter());
+}
+
+
+#undef LOCTEXT_NAMESPACE
diff --git a/Engine/Plugins/Experimental/RenderPages/Source/RenderPagesDeveloper/Private/Graph/RenderPagesGraphNode.cpp b/Engine/Plugins/Experimental/RenderPages/Source/RenderPagesDeveloper/Private/Graph/RenderPagesGraphNode.cpp
new file mode 100644
--- /dev/null
+++ b/Engine/Plugins/Experimental/RenderPages/Source/RenderPagesDeveloper/Private/Graph/RenderPagesGraphNode.cpp
@@ -0,0 +1,3 @@
+// Copyright Epic Games, Inc. All Rights Reserved.
+
+#include "Graph/RenderPagesGraphNode.h"
diff --git a/Engine/Plugins/Experimental/RenderPages/Source/RenderPagesDeveloper/Private/Graph/RenderPagesGraphSchema.cpp b/Engine/Plugins/Experimental/RenderPages/Source/RenderPagesDeveloper/Private/Graph/RenderPagesGraphSchema.cpp
new file mode 100644
--- /dev/null
+++ b/Engine/Plugins/Experimental/RenderPages/Source/RenderPagesDeveloper/Private/Graph/RenderPagesGraphSchema.cpp
@@ -0,0 +1,255 @@
+// Copyright Epic Games, Inc. All Rights Reserved.
+
+#include "Graph/RenderPagesGraphSchema.h"
+#include "Graph/RenderPagesGraph.h"
+#include "Graph/RenderPagesGraphNode.h"
+#include "Blueprints/RenderPagesBlueprint.h"
+#include "EdGraphSchema_K2_Actions.h"
+#include "Kismet2/BlueprintEditorUtils.h"
+#include "Kismet2/Kismet2NameValidators.h"
+
+#define LOCTEXT_NAMESPACE "RenderPagesGraphSchema"
+
+/////////////////////////////////////////////////////
+// UE::RenderPages::FRenderPagesLocalVariableNameValidator
+UE::RenderPages::FRenderPagesLocalVariableNameValidator::FRenderPagesLocalVariableNameValidator(const UBlueprint* Blueprint, const URenderPagesGraph* Graph, FName InExistingName)
+	: FStringSetNameValidator(InExistingName.ToString())
+{
+	if (Blueprint)
+	{
+		TSet<FName> NamesTemp;
+		// We allow local variables with same name as blueprint variable
+
+		FBlueprintEditorUtils::GetFunctionNameList(Blueprint, NamesTemp);
+		FBlueprintEditorUtils::GetAllGraphNames(Blueprint, NamesTemp);
+		FBlueprintEditorUtils::GetSCSVariableNameList(Blueprint, NamesTemp);
+		FBlueprintEditorUtils::GetImplementingBlueprintsFunctionNameList(Blueprint, NamesTemp);
+
+		for (FName& Name : NamesTemp)
+		{
+			Names.Add(Name.ToString());
+		}
+	}
+}
+
+/////////////////////////////////////////////////////
+// UE::RenderPages::FRenderPagesNameValidator
+UE::RenderPages::FRenderPagesNameValidator::FRenderPagesNameValidator(const UBlueprint* Blueprint, const UStruct* ValidationScope, FName InExistingName)
+	: FStringSetNameValidator(InExistingName.ToString())
+{
+	if (Blueprint)
+	{
+		TSet<FName> NamesTemp;
+		FBlueprintEditorUtils::GetClassVariableList(Blueprint, NamesTemp, true);
+		FBlueprintEditorUtils::GetFunctionNameList(Blueprint, NamesTemp);
+		FBlueprintEditorUtils::GetAllGraphNames(Blueprint, NamesTemp);
+		FBlueprintEditorUtils::GetSCSVariableNameList(Blueprint, NamesTemp);
+		FBlueprintEditorUtils::GetImplementingBlueprintsFunctionNameList(Blueprint, NamesTemp);
+
+		for (FName& Name : NamesTemp)
+		{
+			Names.Add(Name.ToString());
+		}
+	}
+}
+
+/////////////////////////////////////////////////////
+// FRenderPagesGraphSchemaAction_LocalVar
+bool FRenderPagesGraphSchemaAction_LocalVar::IsValidName(const FName& NewName, FText& OutErrorMessage) const
+{
+	if (URenderPagesGraph* Graph = Cast<URenderPagesGraph>(GetVariableScope()))
+	{
+		UE::RenderPages::FRenderPagesLocalVariableNameValidator NameValidator(Graph->GetBlueprint(), Graph, GetVariableName());
+		EValidatorResult Result = NameValidator.IsValid(NewName.ToString(), false);
+		if (Result != EValidatorResult::Ok && Result != EValidatorResult::ExistingName)
+		{
+			OutErrorMessage = FText::FromString(TEXT("Name with invalid format"));
+			return false;
+		}
+	}
+	return FEdGraphSchemaAction_BlueprintVariableBase::IsValidName(NewName, OutErrorMessage);
+}
+
+/////////////////////////////////////////////////////
+// FRenderPagesGraphSchemaAction_PromoteToVariable
+FRenderPagesGraphSchemaAction_PromoteToVariable::FRenderPagesGraphSchemaAction_PromoteToVariable(UEdGraphPin* InEdGraphPin, bool InLocalVariable)
+	: FEdGraphSchemaAction(FText(),
+		InLocalVariable ? LOCTEXT("PromoteToLocalVariable", "Promote to local variable") : LOCTEXT("PromoteToVariable", "Promote to variable"),
+		InLocalVariable ? LOCTEXT("PromoteToLocalVariable", "Promote to local variable") : LOCTEXT("PromoteToVariable", "Promote to variable"),
+		1)
+	, EdGraphPin(InEdGraphPin)
+	, bLocalVariable(InLocalVariable)
+{}
+
+UEdGraphNode* FRenderPagesGraphSchemaAction_PromoteToVariable::PerformAction(UEdGraph* ParentGraph, UEdGraphPin* FromPin, const FVector2D Location, bool bSelectNewNode)
+{
+	return nullptr;
+}
+
+/////////////////////////////////////////////////////
+// UE::RenderPages::FRenderPagesFunctionDragDropAction
+UE::RenderPages::FRenderPagesFunctionDragDropAction::FRenderPagesFunctionDragDropAction()
+	: FGraphSchemaActionDragDropAction()
+	, SourceRenderPagesBlueprint(nullptr)
+	, SourceRenderPagesGraph(nullptr)
+	, bControlDrag(false)
+	, bAltDrag(false)
+{}
+
+TSharedRef<UE::RenderPages::FRenderPagesFunctionDragDropAction> UE::RenderPages::FRenderPagesFunctionDragDropAction::New(TSharedPtr<FEdGraphSchemaAction> InAction, URenderPagesBlueprint* InRenderPagesBlueprint, URenderPagesGraph* InRenderPagesGraph)
+{
+	TSharedRef<FRenderPagesFunctionDragDropAction> Action = MakeShareable(new FRenderPagesFunctionDragDropAction);
+	Action->SourceAction = InAction;
+	Action->SourceRenderPagesBlueprint = InRenderPagesBlueprint;
+	Action->SourceRenderPagesGraph = InRenderPagesGraph;
+	Action->Construct();
+	return Action;
+}
+
+/////////////////////////////////////////////////////
+// URenderPagesGraphSchema
+
+const FName URenderPagesGraphSchema::GraphName_RenderPages(TEXT("Logic Graph"));
+
+URenderPagesGraphSchema::URenderPagesGraphSchema() {}
+
+void URenderPagesGraphSchema::GetGraphDisplayInformation(const UEdGraph& Graph, FGraphDisplayInfo& DisplayInfo) const
+{
+	Super::GetGraphDisplayInformation(Graph, DisplayInfo);
+
+	if (Cast<URenderPagesGraph>(const_cast<UEdGraph*>(&Graph)))
+	{
+		static const FText MainGraphText = FText::FromString(TEXT("The logic graph for the Render Pages."));
+		DisplayInfo.Tooltip = MainGraphText;
+	}
+}
+
+FLinearColor URenderPagesGraphSchema::GetPinTypeColor(const FEdGraphPinType& PinType) const
+{
+	return GetDefault<UEdGraphSchema_K2>()->GetPinTypeColor(PinType);
+}
+
+bool URenderPagesGraphSchema::CanGraphBeDropped(TSharedPtr<FEdGraphSchemaAction> InAction) const
+{
+	if (!InAction.IsValid())
+	{
+		return false;
+	}
+
+	if (InAction->GetTypeId() == FEdGraphSchemaAction_K2Graph::StaticGetTypeId())
+	{
+		FEdGraphSchemaAction_K2Graph* FuncAction = static_cast<FEdGraphSchemaAction_K2Graph*>(InAction.Get());
+		if (Cast<URenderPagesGraph>(FuncAction->EdGraph))
+		{
+			return true;
+		}
+	}
+	else if (InAction->GetTypeId() == FRenderPagesGraphSchemaAction_LocalVar::StaticGetTypeId())
+	{
+		FRenderPagesGraphSchemaAction_LocalVar* VarAction = static_cast<FRenderPagesGraphSchemaAction_LocalVar*>(InAction.Get());
+		if (Cast<URenderPagesGraph>(VarAction->GetVariableScope()))
+		{
+			return true;
+		}
+	}
+
+	return false;
+}
+
+FReply URenderPagesGraphSchema::BeginGraphDragAction(TSharedPtr<FEdGraphSchemaAction> InAction, const FPointerEvent& MouseEvent) const
+{
+	if (!InAction.IsValid())
+	{
+		return FReply::Unhandled();
+	}
+
+	if (InAction->GetTypeId() == FEdGraphSchemaAction_K2Graph::StaticGetTypeId())
+	{
+		FEdGraphSchemaAction_K2Graph* FuncAction = static_cast<FEdGraphSchemaAction_K2Graph*>(InAction.Get());
+		if (URenderPagesGraph* RenderPagesGraph = Cast<URenderPagesGraph>(FuncAction->EdGraph))
+		{
+			if (URenderPagesBlueprint* RenderPagesBlueprint = Cast<URenderPagesBlueprint>(FBlueprintEditorUtils::FindBlueprintForGraph(RenderPagesGraph)))
+			{
+				TSharedRef<UE::RenderPages::FRenderPagesFunctionDragDropAction> Action = UE::RenderPages::FRenderPagesFunctionDragDropAction::New(InAction, RenderPagesBlueprint, RenderPagesGraph);
+				Action->SetAltDrag(MouseEvent.IsAltDown());
+				Action->SetCtrlDrag(MouseEvent.IsControlDown());
+				return FReply::Handled().BeginDragDrop(Action);
+			}
+		}
+	}
+	else if (InAction->GetTypeId() == FRenderPagesGraphSchemaAction_LocalVar::StaticGetTypeId())
+	{
+		FRenderPagesGraphSchemaAction_LocalVar* VarAction = static_cast<FRenderPagesGraphSchemaAction_LocalVar*>(InAction.Get());
+		if (URenderPagesGraph* RenderPagesGraph = Cast<URenderPagesGraph>(VarAction->GetVariableScope()))
+		{
+			if (URenderPagesBlueprint* RenderPagesBlueprint = Cast<URenderPagesBlueprint>(FBlueprintEditorUtils::FindBlueprintForGraph(RenderPagesGraph)))
+			{
+				TSharedRef<UE::RenderPages::FRenderPagesFunctionDragDropAction> Action = UE::RenderPages::FRenderPagesFunctionDragDropAction::New(InAction, RenderPagesBlueprint, RenderPagesGraph);
+				Action->SetAltDrag(MouseEvent.IsAltDown());
+				Action->SetCtrlDrag(MouseEvent.IsControlDown());
+				return FReply::Handled().BeginDragDrop(Action);
+			}
+		}
+	}
+	return FReply::Unhandled();
+}
+
+bool URenderPagesGraphSchema::CanVariableBeDropped(UEdGraph* InGraph, FProperty* InVariableToDrop) const
+{
+	return true;
+}
+
+bool URenderPagesGraphSchema::RequestVariableDropOnPanel(UEdGraph* InGraph, FProperty* InVariableToDrop, const FVector2D& InDropPosition, const FVector2D& InScreenPosition)
+{
+	if (CanVariableBeDropped(InGraph, InVariableToDrop))
+	{
+		UBlueprint* Blueprint = FBlueprintEditorUtils::FindBlueprintForGraph(InGraph);
+		if (Cast<URenderPagesBlueprint>(Blueprint))
+		{
+			return true;
+		}
+	}
+	return false;
+}
+
+bool URenderPagesGraphSchema::RequestVariableDropOnPin(UEdGraph* InGraph, FProperty* InVariableToDrop, UEdGraphPin* InPin, const FVector2D& InDropPosition, const FVector2D& InScreenPosition)
+{
+	if (CanVariableBeDropped(InGraph, InVariableToDrop))
+	{
+		if (Cast<URenderPagesGraph>(InGraph))
+		{
+			return true;
+		}
+	}
+	return false;
+}
+
+void URenderPagesGraphSchema::InsertAdditionalActions(TArray<UBlueprint*> InBlueprints, TArray<UEdGraph*> InGraphs, TArray<UEdGraphPin*> InPins, FGraphActionListBuilderBase& OutAllActions) const
+{
+	Super::InsertAdditionalActions(InBlueprints, InGraphs, InPins, OutAllActions);
+
+	if (InPins.Num() > 0)
+	{
+		if (URenderPagesGraphNode* RenderPagesNode = Cast<URenderPagesGraphNode>(InPins[0]->GetOwningNode()))
+		{
+			if (UEdGraphPin* ModelPin = RenderPagesNode->GetPinAt(0))
+			{
+				OutAllActions.AddAction(MakeShared<FRenderPagesGraphSchemaAction_PromoteToVariable>(ModelPin, false));
+			}
+		}
+	}
+}
+
+TSharedPtr<INameValidatorInterface> URenderPagesGraphSchema::GetNameValidator(const UBlueprint* BlueprintObj, const FName& OriginalName, const UStruct* ValidationScope, const FName& ActionTypeId) const
+{
+	if (ActionTypeId == FRenderPagesGraphSchemaAction_LocalVar::StaticGetTypeId())
+	{
+		if (const URenderPagesGraph* RenderPagesGraph = Cast<URenderPagesGraph>(ValidationScope))
+		{
+			return MakeShareable(new UE::RenderPages::FRenderPagesLocalVariableNameValidator(BlueprintObj, RenderPagesGraph, OriginalName));
+		}
+	}
+	return MakeShareable(new UE::RenderPages::FRenderPagesNameValidator(BlueprintObj, ValidationScope, OriginalName));
+}
+
+#undef LOCTEXT_NAMESPACE
diff --git a/Engine/Plugins/Experimental/RenderPages/Source/RenderPagesDeveloper/Private/RenderPagesDeveloperModule.cpp b/Engine/Plugins/Experimental/RenderPages/Source/RenderPagesDeveloper/Private/RenderPagesDeveloperModule.cpp
new file mode 100644
--- /dev/null
+++ b/Engine/Plugins/Experimental/RenderPages/Source/RenderPagesDeveloper/Private/RenderPagesDeveloperModule.cpp
@@ -0,0 +1,14 @@
+// Copyright Epic Games, Inc. All Rights Reserved.
+
+#include "RenderPagesDeveloperModule.h"
+
+#define LOCTEXT_NAMESPACE "RenderPagesDeveloperModule"
+
+
+void UE::RenderPages::Private::FRenderPagesDeveloperModule::StartupModule() {}
+void UE::RenderPages::Private::FRenderPagesDeveloperModule::ShutdownModule() {}
+
+
+#undef LOCTEXT_NAMESPACE
+
+IMPLEMENT_MODULE(UE::RenderPages::Private::FRenderPagesDeveloperModule, RenderPagesDeveloper)
diff --git a/Engine/Plugins/Experimental/RenderPages/Source/RenderPagesDeveloper/Private/RenderPagesDeveloperModule.h b/Engine/Plugins/Experimental/RenderPages/Source/RenderPagesDeveloper/Private/RenderPagesDeveloperModule.h
new file mode 100644
--- /dev/null
+++ b/Engine/Plugins/Experimental/RenderPages/Source/RenderPagesDeveloper/Private/RenderPagesDeveloperModule.h
@@ -0,0 +1,22 @@
+// Copyright Epic Games, Inc. All Rights Reserved.
+
+#pragma once
+
+#include "CoreMinimal.h"
+#include "IRenderPagesDeveloperModule.h"
+
+
+namespace UE::RenderPages::Private
+{
+	/**
+	 * The implementation of the IRenderPagesDeveloperModule interface.
+	 */
+	class FRenderPagesDeveloperModule : public IRenderPagesDeveloperModule
+	{
+	public:
+		//~ Begin IModuleInterface
+		virtual void StartupModule() override;
+		virtual void ShutdownModule() override;
+		//~ End IModuleInterface
+	};
+}
diff --git a/Engine/Plugins/Experimental/RenderPages/Source/RenderPagesDeveloper/Public/Blueprints/RenderPagesBlueprint.h b/Engine/Plugins/Experimental/RenderPages/Source/RenderPagesDeveloper/Public/Blueprints/RenderPagesBlueprint.h
new file mode 100644
--- /dev/null
+++ b/Engine/Plugins/Experimental/RenderPages/Source/RenderPagesDeveloper/Public/Blueprints/RenderPagesBlueprint.h
@@ -0,0 +1,47 @@
+// Copyright Epic Games, Inc. All Rights Reserved.
+
+#pragma once
+
+#include "CoreMinimal.h"
+#include "Engine/Blueprint.h"
+#include "RenderPage/RenderPagePropsSource.h"
+#include "Kismet2/CompilerResultsLog.h"
+#include "RenderPagesBlueprint.generated.h"
+
+
+class URenderPagesBlueprint;
+class URenderPageCollection;
+
+
+/**
+ * A UBlueprint child class for the RenderPages modules.
+ *
+ * Required in order for a RenderPageCollection to be able to have a blueprint graph.
+ */
+UCLASS(BlueprintType, Meta = (IgnoreClassThumbnail))
+class RENDERPAGESDEVELOPER_API URenderPagesBlueprint : public UBlueprint
+{
+	GENERATED_BODY()
+
+public:
+	URenderPagesBlueprint();
+
+public:
+	//~ Begin UBlueprint Interface
+	virtual UClass* GetBlueprintClass() const override;
+	virtual bool SupportedByDefaultBlueprintFactory() const override { return false; }
+	virtual bool IsValidForBytecodeOnlyRecompile() const override { return false; }
+	virtual bool SupportsGlobalVariables() const override { return true; }
+	virtual bool SupportsLocalVariables() const override { return false; }
+	virtual bool SupportsFunctions() const override { return true; }
+	virtual bool SupportsMacros() const override { return false; }
+	virtual bool SupportsDelegates() const override { return true; }
+	virtual bool SupportsEventGraphs() const override { return false; }
+	virtual bool SupportsAnimLayers() const override { return false; }
+
+	virtual void PostLoad() override;
+	//~ End UBlueprint Interface
+
+private:
+	FCompilerResultsLog CompileLog;
+};
diff --git a/Engine/Plugins/Experimental/RenderPages/Source/RenderPagesDeveloper/Public/Blueprints/RenderPagesBlueprintCompiler.h b/Engine/Plugins/Experimental/RenderPages/Source/RenderPagesDeveloper/Public/Blueprints/RenderPagesBlueprintCompiler.h
new file mode 100644
--- /dev/null
+++ b/Engine/Plugins/Experimental/RenderPages/Source/RenderPagesDeveloper/Public/Blueprints/RenderPagesBlueprintCompiler.h
@@ -0,0 +1,24 @@
+// Copyright Epic Games, Inc. All Rights Reserved.
+
+#pragma once
+
+#include "KismetCompilerModule.h"
+#include "KismetCompiler.h"
+
+
+namespace UE::RenderPages
+{
+	/**
+	 * A IBlueprintCompiler child class for the RenderPages modules.
+	 *
+	 * Required in order for a RenderPageCollection to be able to have a blueprint graph.
+	 */
+	class RENDERPAGESDEVELOPER_API FRenderPagesBlueprintCompiler : public IBlueprintCompiler
+	{
+	public:
+		//~ Begin IBlueprintCompiler Interface
+		virtual bool CanCompile(const UBlueprint* Blueprint) override;
+		virtual void Compile(UBlueprint* Blueprint, const FKismetCompilerOptions& CompileOptions, FCompilerResultsLog& Results) override;
+		//~ End IBlueprintCompiler Interface
+	};
+}
diff --git a/Engine/Plugins/Experimental/RenderPages/Source/RenderPagesDeveloper/Public/Graph/RenderPagesGraph.h b/Engine/Plugins/Experimental/RenderPages/Source/RenderPagesDeveloper/Public/Graph/RenderPagesGraph.h
new file mode 100644
--- /dev/null
+++ b/Engine/Plugins/Experimental/RenderPages/Source/RenderPagesDeveloper/Public/Graph/RenderPagesGraph.h
@@ -0,0 +1,34 @@
+// Copyright Epic Games, Inc. All Rights Reserved.
+
+#pragma once
+
+#include "CoreMinimal.h"
+#include "EdGraph/EdGraph.h"
+#include "RenderPagesGraph.generated.h"
+
+
+class URenderPagesBlueprint;
+class URenderPagesGraphSchema;
+class URenderPageCollection;
+
+
+/**
+ * A UEdGraph child class for the RenderPages modules.
+ *
+ * Required in order for a RenderPageCollection to be able to have a blueprint graph.
+ */
+UCLASS()
+class RENDERPAGESDEVELOPER_API URenderPagesGraph : public UEdGraph
+{
+	GENERATED_BODY()
+
+public:
+	/** Set up this graph. */
+	void Initialize(URenderPagesBlueprint* InBlueprint);
+
+	/** Sets the required values for this graph. */
+	virtual void PostLoad() override;
+
+	/** Returns the blueprint of this graph. */
+	URenderPagesBlueprint* GetBlueprint() const;
+};
diff --git a/Engine/Plugins/Experimental/RenderPages/Source/RenderPagesDeveloper/Public/Graph/RenderPagesGraphNode.h b/Engine/Plugins/Experimental/RenderPages/Source/RenderPagesDeveloper/Public/Graph/RenderPagesGraphNode.h
new file mode 100644
--- /dev/null
+++ b/Engine/Plugins/Experimental/RenderPages/Source/RenderPagesDeveloper/Public/Graph/RenderPagesGraphNode.h
@@ -0,0 +1,19 @@
+// Copyright Epic Games, Inc. All Rights Reserved.
+
+#pragma once
+
+#include "CoreMinimal.h"
+#include "EdGraph/EdGraphNode.h"
+#include "RenderPagesGraphNode.generated.h"
+
+
+/**
+ * A UEdGraphNode child class for the RenderPages modules.
+ *
+ * Required in order for a RenderPageCollection to be able to have a blueprint graph.
+ */
+UCLASS()
+class RENDERPAGESDEVELOPER_API URenderPagesGraphNode : public UEdGraphNode
+{
+	GENERATED_BODY()
+};
diff --git a/Engine/Plugins/Experimental/RenderPages/Source/RenderPagesDeveloper/Public/Graph/RenderPagesGraphSchema.h b/Engine/Plugins/Experimental/RenderPages/Source/RenderPagesDeveloper/Public/Graph/RenderPagesGraphSchema.h
new file mode 100644
--- /dev/null
+++ b/Engine/Plugins/Experimental/RenderPages/Source/RenderPagesDeveloper/Public/Graph/RenderPagesGraphSchema.h
@@ -0,0 +1,178 @@
+// Copyright Epic Games, Inc. All Rights Reserved.
+
+#pragma once
+
+#include "CoreMinimal.h"
+#include "EdGraph/EdGraphSchema.h"
+#include "EdGraphSchema_K2_Actions.h"
+#include "GraphEditorDragDropAction.h"
+#include "RenderPagesGraphSchema.generated.h"
+
+
+class URenderPagesBlueprint;
+class URenderPagesGraph;
+
+
+namespace UE::RenderPages
+{
+	/**
+	 * A FStringSetNameValidator child class for the RenderPages modules, for validating local variable names.
+	 *
+	 * Required in order for a RenderPageCollection to be able to have a blueprint graph.
+	 */
+	class RENDERPAGESDEVELOPER_API FRenderPagesLocalVariableNameValidator : public FStringSetNameValidator
+	{
+	public:
+		FRenderPagesLocalVariableNameValidator(const UBlueprint* Blueprint, const URenderPagesGraph* Graph, FName InExistingName = NAME_None);
+	};
+
+	/**
+	 * A FStringSetNameValidator child class for the RenderPages modules, for validating names.
+	 *
+	 * Required in order for a RenderPageCollection to be able to have a blueprint graph.
+	 */
+	class RENDERPAGESDEVELOPER_API FRenderPagesNameValidator : public FStringSetNameValidator
+	{
+	public:
+		FRenderPagesNameValidator(const UBlueprint* Blueprint, const UStruct* ValidationScope, FName InExistingName = NAME_None);
+	};
+}
+
+
+/**
+ * A FEdGraphSchemaAction_BlueprintVariableBase child class for the RenderPages modules, for local variables.
+ *
+ * Required in order for a RenderPageCollection to be able to have a blueprint graph.
+ */
+USTRUCT()
+struct RENDERPAGESDEVELOPER_API FRenderPagesGraphSchemaAction_LocalVar : public FEdGraphSchemaAction_BlueprintVariableBase
+{
+	GENERATED_BODY()
+
+public:
+	/* Returns the type ID. Returns the same string every time. */
+	static FName StaticGetTypeId()
+	{
+		static FName Type("FRenderPagesGraphSchemaAction_LocalVar");
+		return Type;
+	}
+
+	FRenderPagesGraphSchemaAction_LocalVar() = default;
+	FRenderPagesGraphSchemaAction_LocalVar(const FText& InNodeCategory, const FText& InMenuDesc, const FText& InToolTip, const int32 InGrouping, const int32 InSectionID)
+		: FEdGraphSchemaAction_BlueprintVariableBase(InNodeCategory, InMenuDesc, InToolTip, InGrouping, InSectionID)
+	{}
+
+	//~ Begin FEdGraphSchemaAction interface
+	virtual FName GetTypeId() const override { return StaticGetTypeId(); }
+	virtual bool IsA(const FName& InType) const override { return (InType == GetTypeId()) || (InType == FEdGraphSchemaAction_BlueprintVariableBase::StaticGetTypeId()); }
+	virtual bool IsValidName(const FName& NewName, FText& OutErrorMessage) const override;
+	//~ End FEdGraphSchemaAction interface
+};
+
+
+/**
+ * A FEdGraphSchemaAction child class for the RenderPages modules, for the promote-to-variable action.
+ *
+ * Required in order for a RenderPageCollection to be able to have a blueprint graph.
+ */
+USTRUCT()
+struct RENDERPAGESDEVELOPER_API FRenderPagesGraphSchemaAction_PromoteToVariable : public FEdGraphSchemaAction
+{
+	GENERATED_BODY()
+
+public:
+	/* Returns the type ID. Returns the same string every time. */
+	static FName StaticGetTypeId()
+	{
+		static FName Type("FRenderPagesGraphSchemaAction_PromoteToVariable");
+		return Type;
+	}
+
+	FRenderPagesGraphSchemaAction_PromoteToVariable() = default;
+	FRenderPagesGraphSchemaAction_PromoteToVariable(UEdGraphPin* InEdGraphPin, bool InLocalVariable);
+
+	//~ Begin FEdGraphSchemaAction interface
+	virtual FName GetTypeId() const override { return StaticGetTypeId(); }
+	virtual bool IsA(const FName& InType) const override { return (InType == GetTypeId()); }
+	virtual UEdGraphNode* PerformAction(UEdGraph* ParentGraph, UEdGraphPin* FromPin, const FVector2D Location, bool bSelectNewNode) override;
+	//~ End FEdGraphSchemaAction interface
+
+private:
+	TObjectPtr<UEdGraphPin> EdGraphPin = nullptr;
+	bool bLocalVariable = false;
+};
+
+
+namespace UE::RenderPages
+{
+	/**
+	 * A FGraphSchemaActionDragDropAction child class for the RenderPages modules, for drag and dropping an item from the blueprints tree (like a variable or a function).
+	 *
+	 * Required in order for a RenderPageCollection to be able to have a blueprint graph.
+	 */
+	class RENDERPAGESDEVELOPER_API FRenderPagesFunctionDragDropAction : public FGraphSchemaActionDragDropAction
+	{
+	public:
+		DRAG_DROP_OPERATOR_TYPE(FRenderPagesFunctionDragDropAction, FGraphSchemaActionDragDropAction)
+
+	protected:
+		FRenderPagesFunctionDragDropAction();
+
+	public:
+		/** Set if operation is modified by alt. */
+		void SetAltDrag(const bool bInIsAltDrag) { bAltDrag = bInIsAltDrag; }
+
+		/** Set if operation is modified by the ctrl key. */
+		void SetCtrlDrag(const bool bInIsCtrlDrag) { bControlDrag = bInIsCtrlDrag; }
+
+	public:
+		/** Creates and returns a new instance of this class. */
+		static TSharedRef<FRenderPagesFunctionDragDropAction> New(TSharedPtr<FEdGraphSchemaAction> InAction, URenderPagesBlueprint* InRenderPagesBlueprint, URenderPagesGraph* InRenderPagesGraph);
+
+	protected:
+		TObjectPtr<URenderPagesBlueprint> SourceRenderPagesBlueprint;
+		TObjectPtr<URenderPagesGraph> SourceRenderPagesGraph;
+		bool bControlDrag;
+		bool bAltDrag;
+
+		friend class URenderPagesGraphSchema;
+	};
+}
+
+
+/**
+ * A UEdGraphSchema child class for the RenderPages modules.
+ *
+ * Required in order for a RenderPageCollection to be able to have a blueprint graph.
+ */
+UCLASS()
+class RENDERPAGESDEVELOPER_API URenderPagesGraphSchema : public UEdGraphSchema
+{
+	GENERATED_BODY()
+
+public:
+	URenderPagesGraphSchema();
+
+	//~ Begin UEdGraphSchema Interface
+	virtual void GetGraphDisplayInformation(const UEdGraph& Graph, /*out*/ FGraphDisplayInfo& DisplayInfo) const override;
+
+	virtual FLinearColor GetPinTypeColor(const FEdGraphPinType& PinType) const override;
+
+	virtual bool CanDuplicateGraph(UEdGraph* InSourceGraph) const override { return false; }
+
+	virtual bool CanGraphBeDropped(TSharedPtr<FEdGraphSchemaAction> InAction) const override;
+	virtual FReply BeginGraphDragAction(TSharedPtr<FEdGraphSchemaAction> InAction, const FPointerEvent& MouseEvent = FPointerEvent()) const override;
+
+	virtual bool CanVariableBeDropped(UEdGraph* InGraph, FProperty* InVariableToDrop) const override;
+	virtual bool RequestVariableDropOnPanel(UEdGraph* InGraph, FProperty* InVariableToDrop, const FVector2D& InDropPosition, const FVector2D& InScreenPosition) override;
+	virtual bool RequestVariableDropOnPin(UEdGraph* InGraph, FProperty* InVariableToDrop, UEdGraphPin* InPin, const FVector2D& InDropPosition, const FVector2D& InScreenPosition) override;
+
+	virtual void InsertAdditionalActions(TArray<UBlueprint*> InBlueprints, TArray<UEdGraph*> InGraphs, TArray<UEdGraphPin*> InPins, FGraphActionListBuilderBase& OutAllActions) const override;
+
+	virtual TSharedPtr<INameValidatorInterface> GetNameValidator(const UBlueprint* BlueprintObj, const FName& OriginalName, const UStruct* ValidationScope, const FName& ActionTypeId) const override;
+	//~ End UEdGraphSchema Interface
+
+public:
+	/** Name constant. */
+	static const FName GraphName_RenderPages;
+};
diff --git a/Engine/Plugins/Experimental/RenderPages/Source/RenderPagesDeveloper/Public/IRenderPagesDeveloperModule.h b/Engine/Plugins/Experimental/RenderPages/Source/RenderPagesDeveloper/Public/IRenderPagesDeveloperModule.h
new file mode 100644
--- /dev/null
+++ b/Engine/Plugins/Experimental/RenderPages/Source/RenderPagesDeveloper/Public/IRenderPagesDeveloperModule.h
@@ -0,0 +1,15 @@
+// Copyright Epic Games, Inc. All Rights Reserved.
+
+#pragma once
+
+#include "CoreMinimal.h"
+#include "Modules/ModuleManager.h"
+
+
+namespace UE::RenderPages
+{
+	/**
+	 * RenderPagesDeveloper module interface.
+	 */
+	class IRenderPagesDeveloperModule : public IModuleInterface	{};
+}
diff --git a/Engine/Plugins/Experimental/RenderPages/Source/RenderPagesDeveloper/Public/RenderPagesDeveloperLog.h b/Engine/Plugins/Experimental/RenderPages/Source/RenderPagesDeveloper/Public/RenderPagesDeveloperLog.h
new file mode 100644
--- /dev/null
+++ b/Engine/Plugins/Experimental/RenderPages/Source/RenderPagesDeveloper/Public/RenderPagesDeveloperLog.h
@@ -0,0 +1,8 @@
+// Copyright Epic Games, Inc. All Rights Reserved.
+
+#pragma once
+
+#include "CoreMinimal.h"
+
+
+DECLARE_LOG_CATEGORY_CLASS(LogRenderPagesDeveloper, Log, All);
diff --git a/Engine/Plugins/Experimental/RenderPages/Source/RenderPagesDeveloper/RenderPagesDeveloper.Build.cs b/Engine/Plugins/Experimental/RenderPages/Source/RenderPagesDeveloper/RenderPagesDeveloper.Build.cs
new file mode 100644
--- /dev/null
+++ b/Engine/Plugins/Experimental/RenderPages/Source/RenderPagesDeveloper/RenderPagesDeveloper.Build.cs
@@ -0,0 +1,32 @@
+// Copyright Epic Games, Inc. All Rights Reserved.
+
+using UnrealBuildTool;
+
+public class RenderPagesDeveloper : ModuleRules
+{
+	public RenderPagesDeveloper(ReadOnlyTargetRules Target) : base(Target)
+	{
+		PCHUsage = PCHUsageMode.UseExplicitOrSharedPCHs;
+
+		PublicDependencyModuleNames.AddRange(
+			new string[]
+			{
+				"Core",
+				"CoreUObject",
+			}
+		);
+
+		PrivateDependencyModuleNames.AddRange(
+			new string[]
+			{
+				"BlueprintGraph",
+				"Engine",
+				"GraphEditor",
+				"RenderPages",
+				"Slate",
+				"SlateCore",
+				"UnrealEd",
+			}
+		);
+	}
+}
\ No newline at end of file
diff --git a/Engine/Plugins/Experimental/RenderPages/Source/RenderPagesEditor/Private/AssetTypeActions/RenderPagesBlueprintActions.cpp b/Engine/Plugins/Experimental/RenderPages/Source/RenderPagesEditor/Private/AssetTypeActions/RenderPagesBlueprintActions.cpp
new file mode 100644
--- /dev/null
+++ b/Engine/Plugins/Experimental/RenderPages/Source/RenderPagesEditor/Private/AssetTypeActions/RenderPagesBlueprintActions.cpp
@@ -0,0 +1,77 @@
+// Copyright Epic Games, Inc. All Rights Reserved.
+
+#include "RenderPagesBlueprintActions.h"
+#include "Styling/AppStyle.h"
+#include "Blueprints/RenderPagesBlueprint.h"
+#include "Factories/RenderPageCollectionFactory.h"
+#include "RenderPagesEditorModule.h"
+#include "Styling/SlateIconFinder.h"
+
+#define LOCTEXT_NAMESPACE "AssetTypeActions"
+
+
+FText UE::RenderPages::Private::FRenderPagesBlueprintActions::GetName() const
+{
+	return LOCTEXT("AssetTypeActions_RenderPageCollection", "Render Page Collection");
+}
+
+FColor UE::RenderPages::Private::FRenderPagesBlueprintActions::GetTypeColor() const
+{
+	return FColor(255, 64, 64);
+}
+
+UClass* UE::RenderPages::Private::FRenderPagesBlueprintActions::GetSupportedClass() const
+{
+	return URenderPagesBlueprint::StaticClass();
+}
+
+void UE::RenderPages::Private::FRenderPagesBlueprintActions::OpenAssetEditor(const TArray<UObject*>& InObjects, TSharedPtr<IToolkitHost> EditWithinLevelEditor)
+{
+	for (auto ObjIt = InObjects.CreateConstIterator(); ObjIt; ++ObjIt)
+	{
+		if (URenderPagesBlueprint* RenderPagesBlueprint = Cast<URenderPagesBlueprint>(*ObjIt))
+		{
+			constexpr bool bBringToFrontIfOpen = true;
+			if (IAssetEditorInstance* EditorInstance = GEditor->GetEditorSubsystem<UAssetEditorSubsystem>()->FindEditorForAsset(RenderPagesBlueprint, bBringToFrontIfOpen))
+			{
+				EditorInstance->FocusWindow(RenderPagesBlueprint);
+			}
+			else
+			{
+				const EToolkitMode::Type Mode = EditWithinLevelEditor.IsValid() ? EToolkitMode::WorldCentric : EToolkitMode::Standalone;
+				IRenderPagesEditorModule::Get().CreateRenderPageCollectionEditor(Mode, EditWithinLevelEditor, RenderPagesBlueprint);
+			}
+		}
+	}
+}
+
+uint32 UE::RenderPages::Private::FRenderPagesBlueprintActions::GetCategories()
+{
+	return EAssetTypeCategories::Misc;
+}
+
+TSharedPtr<SWidget> UE::RenderPages::Private::FRenderPagesBlueprintActions::GetThumbnailOverlay(const FAssetData& AssetData) const
+{
+	const FSlateBrush* Icon = FSlateIconFinder::FindIconBrushForClass(URenderPagesBlueprint::StaticClass());
+
+	return SNew(SBorder)
+		.BorderImage(FAppStyle::GetNoBrush())
+		.Visibility(EVisibility::HitTestInvisible)
+		.Padding(FMargin(0.0f, 0.0f, 0.0f, 3.0f))
+		.HAlign(HAlign_Right)
+		.VAlign(VAlign_Bottom)
+		[
+			SNew(SImage)
+			.Image(Icon)
+		];
+}
+
+UFactory* UE::RenderPages::Private::FRenderPagesBlueprintActions::GetFactoryForBlueprintType(UBlueprint* InBlueprint) const
+{
+	URenderPagesBlueprintFactory* RenderPagesBlueprintFactory = NewObject<URenderPagesBlueprintFactory>();
+	RenderPagesBlueprintFactory->ParentClass = TSubclassOf<URenderPageCollection>(*InBlueprint->GeneratedClass);
+	return RenderPagesBlueprintFactory;
+}
+
+
+#undef LOCTEXT_NAMESPACE
diff --git a/Engine/Plugins/Experimental/RenderPages/Source/RenderPagesEditor/Private/AssetTypeActions/RenderPagesBlueprintActions.h b/Engine/Plugins/Experimental/RenderPages/Source/RenderPagesEditor/Private/AssetTypeActions/RenderPagesBlueprintActions.h
new file mode 100644
--- /dev/null
+++ b/Engine/Plugins/Experimental/RenderPages/Source/RenderPagesEditor/Private/AssetTypeActions/RenderPagesBlueprintActions.h
@@ -0,0 +1,30 @@
+// Copyright Epic Games, Inc. All Rights Reserved.
+
+#pragma once
+
+#include "AssetTypeActions_Base.h"
+#include "AssetTypeActions/AssetTypeActions_Blueprint.h"
+
+
+namespace UE::RenderPages::Private
+{
+	/**
+	 * This class adds support for the RenderPageCollection (RenderPagesBlueprint) class to the AssetTools module.
+	 */
+	class FRenderPagesBlueprintActions : public FAssetTypeActions_Blueprint
+	{
+	public:
+		//~ Begin IAssetTypeActions Interface
+		virtual FText GetName() const override;
+		virtual FColor GetTypeColor() const override;
+		virtual UClass* GetSupportedClass() const override;
+		virtual void OpenAssetEditor(const TArray<UObject*>& InObjects, TSharedPtr<IToolkitHost> EditWithinLevelEditor = TSharedPtr<IToolkitHost>()) override;
+		virtual uint32 GetCategories() override;
+		virtual TSharedPtr<SWidget> GetThumbnailOverlay(const FAssetData& AssetData) const override;
+		//~ End IAssetTypeActions Interface
+
+		//~ Begin FAssetTypeActions_Blueprint interface
+		virtual UFactory* GetFactoryForBlueprintType(UBlueprint* InBlueprint) const override;
+		//~ End FAssetTypeActions_Blueprint interface
+	};
+}
diff --git a/Engine/Plugins/Experimental/RenderPages/Source/RenderPagesEditor/Private/BlueprintModes/RenderPagesApplicationModeBase.cpp b/Engine/Plugins/Experimental/RenderPages/Source/RenderPagesEditor/Private/BlueprintModes/RenderPagesApplicationModeBase.cpp
new file mode 100644
--- /dev/null
+++ b/Engine/Plugins/Experimental/RenderPages/Source/RenderPagesEditor/Private/BlueprintModes/RenderPagesApplicationModeBase.cpp
@@ -0,0 +1,25 @@
+// Copyright Epic Games, Inc. All Rights Reserved.
+
+#include "BlueprintModes/RenderPagesApplicationModeBase.h"
+#include "BlueprintModes/RenderPagesApplicationModes.h"
+#include "IRenderPageCollectionEditor.h"
+
+
+UE::RenderPages::Private::FRenderPagesApplicationModeBase::FRenderPagesApplicationModeBase(TSharedPtr<IRenderPageCollectionEditor> InRenderPagesEditor, FName InModeName)
+	: FBlueprintEditorApplicationMode(InRenderPagesEditor, InModeName, FRenderPagesApplicationModes::GetLocalizedMode, false, false)
+	, BlueprintEditorWeakPtr(InRenderPagesEditor)
+{}
+
+URenderPagesBlueprint* UE::RenderPages::Private::FRenderPagesApplicationModeBase::GetBlueprint() const
+{
+	if (const TSharedPtr<IRenderPageCollectionEditor> BlueprintEditor = BlueprintEditorWeakPtr.Pin())
+	{
+		return BlueprintEditor->GetRenderPagesBlueprint();
+	}
+	return nullptr;
+}
+
+TSharedPtr<UE::RenderPages::IRenderPageCollectionEditor> UE::RenderPages::Private::FRenderPagesApplicationModeBase::GetBlueprintEditor() const
+{
+	return BlueprintEditorWeakPtr.Pin();
+}
diff --git a/Engine/Plugins/Experimental/RenderPages/Source/RenderPagesEditor/Private/BlueprintModes/RenderPagesApplicationModeBase.h b/Engine/Plugins/Experimental/RenderPages/Source/RenderPagesEditor/Private/BlueprintModes/RenderPagesApplicationModeBase.h
new file mode 100644
--- /dev/null
+++ b/Engine/Plugins/Experimental/RenderPages/Source/RenderPagesEditor/Private/BlueprintModes/RenderPagesApplicationModeBase.h
@@ -0,0 +1,44 @@
+// Copyright Epic Games, Inc. All Rights Reserved.
+
+#pragma once
+
+#include "CoreMinimal.h"
+#include "BlueprintEditorModes.h"
+#include "WorkflowOrientedApp/WorkflowTabManager.h"
+
+
+class URenderPagesBlueprint;
+
+namespace UE::RenderPages
+{
+	class IRenderPageCollectionEditor;
+}
+
+
+namespace UE::RenderPages::Private
+{
+	/**
+	 * This is the base class for the render page editor application modes.
+	 * 
+	 * It contains functionality that's shared between all the render page editor application modes.
+	 */
+	class FRenderPagesApplicationModeBase : public FBlueprintEditorApplicationMode
+	{
+	public:
+		FRenderPagesApplicationModeBase(TSharedPtr<IRenderPageCollectionEditor> InRenderPagesEditor, FName InModeName);
+
+	protected:
+		/** Returns the RenderPagesBlueprint of the editor that was given to the constructor. */
+		URenderPagesBlueprint* GetBlueprint() const;
+
+		/** Returns the editor that was given to the constructor. */
+		TSharedPtr<IRenderPageCollectionEditor> GetBlueprintEditor() const;
+
+	protected:
+		/** The editor that was given to the constructor. */
+		TWeakPtr<IRenderPageCollectionEditor> BlueprintEditorWeakPtr;
+
+		/** Set of spawnable tabs in the mode. */
+		FWorkflowAllowedTabSet TabFactories;
+	};
+}
diff --git a/Engine/Plugins/Experimental/RenderPages/Source/RenderPagesEditor/Private/BlueprintModes/RenderPagesApplicationModeListing.cpp b/Engine/Plugins/Experimental/RenderPages/Source/RenderPagesEditor/Private/BlueprintModes/RenderPagesApplicationModeListing.cpp
new file mode 100644
--- /dev/null
+++ b/Engine/Plugins/Experimental/RenderPages/Source/RenderPagesEditor/Private/BlueprintModes/RenderPagesApplicationModeListing.cpp
@@ -0,0 +1,113 @@
+// Copyright Epic Games, Inc. All Rights Reserved.
+
+#include "BlueprintModes/RenderPagesApplicationModeListing.h"
+#include "BlueprintModes/RenderPagesApplicationModes.h"
+
+#include "BlueprintEditorSharedTabFactories.h"
+#include "BlueprintEditorTabs.h"
+
+#include "IRenderPageCollectionEditor.h"
+#include "IRenderPagesEditorModule.h"
+
+#include "SBlueprintEditorToolbar.h"
+
+#include "TabFactory/PageListTabSummoner.h"
+#include "TabFactory/CollectionPropertiesTabSummoner.h"
+#include "TabFactory/PagePropertiesTabSummoner.h"
+#include "TabFactory/PageViewerTabSummoner.h"
+
+#include "Toolkit/RenderPageCollectionEditorToolbar.h"
+
+#define LOCTEXT_NAMESPACE "RenderPagesListingMode"
+
+
+UE::RenderPages::Private::FRenderPagesApplicationModeListing::FRenderPagesApplicationModeListing(TSharedPtr<IRenderPageCollectionEditor> InRenderPagesEditor)
+	: FRenderPagesApplicationModeBase(InRenderPagesEditor, FRenderPagesApplicationModes::ListingMode)
+{
+	// Override the default created category here since "Listing Editor" sounds awkward
+	WorkspaceMenuCategory = FWorkspaceItem::NewGroup(LOCTEXT("WorkspaceMenu_RenderPagesListing", "Render Pages Listing"));
+
+	TabLayout = FTabManager::NewLayout("RenderPagesBlueprintEditor_Listing_Layout_v1_000")
+		->AddArea(
+			FTabManager::NewPrimaryArea()
+			->SetOrientation(Orient_Vertical)
+			->Split(
+				FTabManager::NewSplitter()
+				->SetSizeCoefficient(0.5f)
+				->SetOrientation(Orient_Horizontal)
+				->Split(
+					FTabManager::NewStack()
+					->SetHideTabWell(true)
+					->SetSizeCoefficient(1.f)
+					->SetForegroundTab(FPageListTabSummoner::TabID)
+					->AddTab(FPageListTabSummoner::TabID, ETabState::OpenedTab)
+				)
+			)
+			->Split(
+				FTabManager::NewSplitter()
+				->SetSizeCoefficient(0.5f)
+				->SetOrientation(Orient_Horizontal)
+				->Split(
+					FTabManager::NewStack()
+					->SetHideTabWell(true)
+					->SetSizeCoefficient(0.25f)
+					->SetForegroundTab(FCollectionPropertiesTabSummoner::TabID)
+					->AddTab(FCollectionPropertiesTabSummoner::TabID, ETabState::OpenedTab)
+				)
+				->Split(
+					FTabManager::NewStack()
+					->SetHideTabWell(true)
+					->SetSizeCoefficient(0.5f)
+					->SetForegroundTab(FPageViewerTabSummoner::TabID)
+					->AddTab(FPageViewerTabSummoner::TabID, ETabState::OpenedTab)
+				)
+				->Split(
+					FTabManager::NewStack()
+					->SetHideTabWell(true)
+					->SetSizeCoefficient(0.25f)
+					->SetForegroundTab(FPagePropertiesTabSummoner::TabID)
+					->AddTab(FPagePropertiesTabSummoner::TabID, ETabState::OpenedTab)
+				)
+			)
+		);
+
+	// Register Tab Factories
+	TabFactories.RegisterFactory(MakeShareable(new FPageListTabSummoner(InRenderPagesEditor)));
+	TabFactories.RegisterFactory(MakeShareable(new FCollectionPropertiesTabSummoner(InRenderPagesEditor)));
+	TabFactories.RegisterFactory(MakeShareable(new FPageViewerTabSummoner(InRenderPagesEditor)));
+	TabFactories.RegisterFactory(MakeShareable(new FPagePropertiesTabSummoner(InRenderPagesEditor)));
+
+	//Make sure we start with our existing list of extenders instead of creating a new one
+	IRenderPagesEditorModule& RenderPagesEditorModule = IRenderPagesEditorModule::Get();
+	ToolbarExtender = RenderPagesEditorModule.GetToolBarExtensibilityManager()->GetAllExtenders();
+
+	InRenderPagesEditor->GetRenderPagesToolbarBuilder()->AddRenderPagesBlueprintEditorModesToolbar(ToolbarExtender);
+
+	if (UToolMenu* Toolbar = InRenderPagesEditor->RegisterModeToolbarIfUnregistered(GetModeName()))
+	{
+		InRenderPagesEditor->GetRenderPagesToolbarBuilder()->AddListingModeToolbar(Toolbar);
+
+		InRenderPagesEditor->GetToolbarBuilder()->AddCompileToolbar(Toolbar);
+	}
+}
+
+void UE::RenderPages::Private::FRenderPagesApplicationModeListing::RegisterTabFactories(TSharedPtr<FTabManager> InTabManager)
+{
+	TSharedPtr<IRenderPageCollectionEditor> BP = GetBlueprintEditor();
+
+	BP->RegisterToolbarTab(InTabManager.ToSharedRef());
+	BP->PushTabFactories(TabFactories);
+}
+
+void UE::RenderPages::Private::FRenderPagesApplicationModeListing::PreDeactivateMode()
+{
+	// prevented from running: FRenderPagesBlueprintApplicationMode::PreDeactivateMode();
+}
+
+void UE::RenderPages::Private::FRenderPagesApplicationModeListing::PostActivateMode()
+{
+	// prevented from running: FRenderPagesBlueprintApplicationMode::PostActivateMode();
+}
+
+
+#undef LOCTEXT_NAMESPACE
diff --git a/Engine/Plugins/Experimental/RenderPages/Source/RenderPagesEditor/Private/BlueprintModes/RenderPagesApplicationModeListing.h b/Engine/Plugins/Experimental/RenderPages/Source/RenderPagesEditor/Private/BlueprintModes/RenderPagesApplicationModeListing.h
new file mode 100644
--- /dev/null
+++ b/Engine/Plugins/Experimental/RenderPages/Source/RenderPagesEditor/Private/BlueprintModes/RenderPagesApplicationModeListing.h
@@ -0,0 +1,25 @@
+// Copyright Epic Games, Inc. All Rights Reserved.
+
+#pragma once
+
+#include "CoreMinimal.h"
+#include "BlueprintModes/RenderPagesApplicationModeBase.h"
+
+
+namespace UE::RenderPages::Private
+{
+	/**
+	 * This is the application mode for the render page editor listing functionality (the list of pages, with the render previews, the page properties, etc).
+	 */
+	class FRenderPagesApplicationModeListing : public FRenderPagesApplicationModeBase
+	{
+	public:
+		FRenderPagesApplicationModeListing(TSharedPtr<IRenderPageCollectionEditor> InRenderPagesEditor);
+
+		//~ Begin FApplicationMode interface
+		virtual void RegisterTabFactories(TSharedPtr<FTabManager> InTabManager) override;
+		virtual void PreDeactivateMode() override;
+		virtual void PostActivateMode() override;
+		//~ End FApplicationMode interface
+	};
+}
diff --git a/Engine/Plugins/Experimental/RenderPages/Source/RenderPagesEditor/Private/BlueprintModes/RenderPagesApplicationModeLogic.cpp b/Engine/Plugins/Experimental/RenderPages/Source/RenderPagesEditor/Private/BlueprintModes/RenderPagesApplicationModeLogic.cpp
new file mode 100644
--- /dev/null
+++ b/Engine/Plugins/Experimental/RenderPages/Source/RenderPagesEditor/Private/BlueprintModes/RenderPagesApplicationModeLogic.cpp
@@ -0,0 +1,104 @@
+// Copyright Epic Games, Inc. All Rights Reserved.
+
+#include "BlueprintModes/RenderPagesApplicationModeLogic.h"
+#include "BlueprintModes/RenderPagesApplicationModes.h"
+
+#include "BlueprintEditorSharedTabFactories.h"
+#include "BlueprintEditorTabs.h"
+
+#include "IRenderPagesEditorModule.h"
+#include "IRenderPageCollectionEditor.h"
+
+#include "SBlueprintEditorToolbar.h"
+
+#include "Toolkit/RenderPageCollectionEditorToolbar.h"
+
+#define LOCTEXT_NAMESPACE "RenderPagesLogicMode"
+
+
+UE::RenderPages::Private::FRenderPagesApplicationModeLogic::FRenderPagesApplicationModeLogic(TSharedPtr<IRenderPageCollectionEditor> InRenderPagesEditor)
+	: FRenderPagesApplicationModeBase(InRenderPagesEditor, FRenderPagesApplicationModes::LogicMode)
+{
+	// Override the default created category here since "Logic Editor" sounds awkward
+	WorkspaceMenuCategory = FWorkspaceItem::NewGroup(LOCTEXT("WorkspaceMenu_RenderPagesLogic", "Render Pages Logic"));
+
+	TabLayout = FTabManager::NewLayout("RenderPagesBlueprintEditor_Logic_Layout_v1")
+		->AddArea(
+			FTabManager::NewPrimaryArea()->SetOrientation(Orient_Vertical)
+			->Split(
+				FTabManager::NewSplitter()->SetOrientation(Orient_Horizontal)
+				->Split(
+					FTabManager::NewSplitter()->SetOrientation(Orient_Vertical)
+					->SetSizeCoefficient(0.15f)
+					->Split(
+						FTabManager::NewStack()->SetSizeCoefficient(0.5f)
+						->AddTab(FBlueprintEditorTabs::MyBlueprintID, ETabState::OpenedTab)
+					)
+					->Split(
+						FTabManager::NewStack()->SetSizeCoefficient(0.5f)
+						->AddTab(FBlueprintEditorTabs::DetailsID, ETabState::OpenedTab)
+					)
+				)
+				->Split(
+					FTabManager::NewSplitter()->SetOrientation(Orient_Vertical)
+					->SetSizeCoefficient(0.70f)
+					->Split(
+						FTabManager::NewStack()
+						->SetSizeCoefficient(0.80f)
+						->AddTab("Document", ETabState::ClosedTab)
+					)
+					->Split(
+						FTabManager::NewStack()
+						->SetSizeCoefficient(0.20f)
+						->AddTab(FBlueprintEditorTabs::CompilerResultsID, ETabState::ClosedTab)
+						->AddTab(FBlueprintEditorTabs::FindResultsID, ETabState::ClosedTab)
+					)
+				)
+				->Split(
+					FTabManager::NewSplitter()->SetOrientation(Orient_Vertical)
+					->SetSizeCoefficient(0.15f)
+					->Split(
+						FTabManager::NewStack()
+						->AddTab(FBlueprintEditorTabs::PaletteID, ETabState::ClosedTab)
+					)
+				)
+			)
+		);
+
+	//Make sure we start with new one
+	ToolbarExtender = MakeShareable(new FExtender);
+
+	InRenderPagesEditor->GetRenderPagesToolbarBuilder()->AddRenderPagesBlueprintEditorModesToolbar(ToolbarExtender);
+
+	if (UToolMenu* Toolbar = InRenderPagesEditor->RegisterModeToolbarIfUnregistered(GetModeName()))
+	{
+		InRenderPagesEditor->GetRenderPagesToolbarBuilder()->AddLogicModeToolbar(Toolbar);
+
+		InRenderPagesEditor->GetToolbarBuilder()->AddCompileToolbar(Toolbar);
+		InRenderPagesEditor->GetToolbarBuilder()->AddScriptingToolbar(Toolbar);
+		InRenderPagesEditor->GetToolbarBuilder()->AddBlueprintGlobalOptionsToolbar(Toolbar);
+	}
+}
+
+void UE::RenderPages::Private::FRenderPagesApplicationModeLogic::RegisterTabFactories(TSharedPtr<FTabManager> InTabManager)
+{
+	TSharedPtr<FBlueprintEditor> BP = GetBlueprintEditor();
+
+	BP->RegisterToolbarTab(InTabManager.ToSharedRef());
+	BP->PushTabFactories(CoreTabFactories);
+	BP->PushTabFactories(BlueprintEditorTabFactories);
+	BP->PushTabFactories(TabFactories);
+}
+
+void UE::RenderPages::Private::FRenderPagesApplicationModeLogic::PreDeactivateMode()
+{
+	// prevented from running: FRenderPagesBlueprintApplicationMode::PreDeactivateMode();
+}
+
+void UE::RenderPages::Private::FRenderPagesApplicationModeLogic::PostActivateMode()
+{
+	FRenderPagesApplicationModeBase::PostActivateMode();
+}
+
+
+#undef LOCTEXT_NAMESPACE
diff --git a/Engine/Plugins/Experimental/RenderPages/Source/RenderPagesEditor/Private/BlueprintModes/RenderPagesApplicationModeLogic.h b/Engine/Plugins/Experimental/RenderPages/Source/RenderPagesEditor/Private/BlueprintModes/RenderPagesApplicationModeLogic.h
new file mode 100644
--- /dev/null
+++ b/Engine/Plugins/Experimental/RenderPages/Source/RenderPagesEditor/Private/BlueprintModes/RenderPagesApplicationModeLogic.h
@@ -0,0 +1,26 @@
+// Copyright Epic Games, Inc. All Rights Reserved.
+
+#pragma once
+
+#include "CoreMinimal.h"
+#include "BlueprintModes/RenderPagesApplicationModeBase.h"
+#include "WorkflowOrientedApp/WorkflowTabManager.h"
+
+
+namespace UE::RenderPages::Private
+{
+	/**
+	 * This is the application mode for the render page editor 'logic' functionality (the blueprint graph).
+	 */
+	class FRenderPagesApplicationModeLogic : public FRenderPagesApplicationModeBase
+	{
+	public:
+		FRenderPagesApplicationModeLogic(TSharedPtr<IRenderPageCollectionEditor> InRenderPagesEditor);
+
+		//~ Begin FApplicationMode interface
+		virtual void RegisterTabFactories(TSharedPtr<FTabManager> InTabManager) override;
+		virtual void PreDeactivateMode() override;
+		virtual void PostActivateMode() override;
+		//~ End FApplicationMode interface
+	};
+}
diff --git a/Engine/Plugins/Experimental/RenderPages/Source/RenderPagesEditor/Private/BlueprintModes/RenderPagesApplicationModes.cpp b/Engine/Plugins/Experimental/RenderPages/Source/RenderPagesEditor/Private/BlueprintModes/RenderPagesApplicationModes.cpp
new file mode 100644
--- /dev/null
+++ b/Engine/Plugins/Experimental/RenderPages/Source/RenderPagesEditor/Private/BlueprintModes/RenderPagesApplicationModes.cpp
@@ -0,0 +1,25 @@
+// Copyright Epic Games, Inc. All Rights Reserved.
+
+#include "BlueprintModes/RenderPagesApplicationModes.h"
+
+
+// Mode constants
+const FName UE::RenderPages::Private::FRenderPagesApplicationModes::ListingMode("ListingName");
+const FName UE::RenderPages::Private::FRenderPagesApplicationModes::LogicMode("LogicName");
+
+
+FText UE::RenderPages::Private::FRenderPagesApplicationModes::GetLocalizedMode(const FName InMode)
+{
+	static TMap<FName, FText> LocModes;
+	if (LocModes.Num() == 0)
+	{
+		LocModes.Add(ListingMode, NSLOCTEXT("RenderPagesBlueprintModes", "ListingMode", "Listing"));
+		LocModes.Add(LogicMode, NSLOCTEXT("RenderPagesBlueprintModes", "LogicMode", "Logic"));
+	}
+	
+	check(InMode != NAME_None);
+	const FText* OutDesc = LocModes.Find(InMode);
+	check(OutDesc);
+	
+	return *OutDesc;
+}
diff --git a/Engine/Plugins/Experimental/RenderPages/Source/RenderPagesEditor/Private/BlueprintModes/RenderPagesApplicationModes.h b/Engine/Plugins/Experimental/RenderPages/Source/RenderPagesEditor/Private/BlueprintModes/RenderPagesApplicationModes.h
new file mode 100644
--- /dev/null
+++ b/Engine/Plugins/Experimental/RenderPages/Source/RenderPagesEditor/Private/BlueprintModes/RenderPagesApplicationModes.h
@@ -0,0 +1,29 @@
+// Copyright Epic Games, Inc. All Rights Reserved.
+
+#pragma once
+
+#include "CoreMinimal.h"
+
+
+namespace UE::RenderPages::Private
+{
+	/**
+	 * This struct contains the IDs as well as the localized text of each of the render page editor application modes.
+	 */
+	struct FRenderPagesApplicationModes
+	{
+	public:
+		/** Constant for the listing mode. */
+		static const FName ListingMode;
+
+		/** Constant for the logic mode. */
+		static const FName LogicMode;
+
+	public:
+		/** Returns the localized text for the given render page editor application mode. */
+		static FText GetLocalizedMode(const FName InMode);
+
+	private:
+		FRenderPagesApplicationModes() = default;
+	};
+}
diff --git a/Engine/Plugins/Experimental/RenderPages/Source/RenderPagesEditor/Private/Commands/RenderPagesEditorCommands.cpp b/Engine/Plugins/Experimental/RenderPages/Source/RenderPagesEditor/Private/Commands/RenderPagesEditorCommands.cpp
new file mode 100644
--- /dev/null
+++ b/Engine/Plugins/Experimental/RenderPages/Source/RenderPagesEditor/Private/Commands/RenderPagesEditorCommands.cpp
@@ -0,0 +1,20 @@
+// Copyright Epic Games, Inc. All Rights Reserved.
+
+#include "Commands/RenderPagesEditorCommands.h"
+
+#define LOCTEXT_NAMESPACE "FRenderPagesEditor"
+
+
+void UE::RenderPages::Private::FRenderPagesEditorCommands::RegisterCommands()
+{
+	UI_COMMAND(AddPage, "Add", "Adds a new page instance to the list.", EUserInterfaceActionType::Button, FInputChord());
+	UI_COMMAND(CopyPage, "Copy", "Adds a new page instance to the list by copying from an existing instance.", EUserInterfaceActionType::Button, FInputChord());
+	UI_COMMAND(DeletePage, "Delete", "Deletes an existing page instance from the list.", EUserInterfaceActionType::Button, FInputChord());
+
+	UI_COMMAND(BatchRenderList, "Render", "Renders the available page instance(s) in batch.", EUserInterfaceActionType::Button, FInputChord());
+
+	UI_COMMAND(AutoCompileGraph, "Auto Compile", "Auto-compile the render pages graph if turned on.", EUserInterfaceActionType::ToggleButton, FInputChord());
+}
+
+
+#undef LOCTEXT_NAMESPACE
diff --git a/Engine/Plugins/Experimental/RenderPages/Source/RenderPagesEditor/Private/Commands/RenderPagesEditorCommands.h b/Engine/Plugins/Experimental/RenderPages/Source/RenderPagesEditor/Private/Commands/RenderPagesEditorCommands.h
new file mode 100644
--- /dev/null
+++ b/Engine/Plugins/Experimental/RenderPages/Source/RenderPagesEditor/Private/Commands/RenderPagesEditorCommands.h
@@ -0,0 +1,43 @@
+// Copyright Epic Games, Inc. All Rights Reserved.
+
+#pragma once
+
+#include "CoreMinimal.h"
+#include "Framework/Commands/Commands.h"
+#include "Styling/AppStyle.h"
+
+
+namespace UE::RenderPages::Private
+{
+	/**
+	 * Contains the commands for the render page editor.
+	 */
+	class FRenderPagesEditorCommands : public TCommands<FRenderPagesEditorCommands>
+	{
+	public:
+		FRenderPagesEditorCommands()
+			: TCommands<FRenderPagesEditorCommands>(TEXT("RenderPagesEditor"),
+				NSLOCTEXT("Contexts", "RenderPagesEditor", "Render Pages Editor"),
+				NAME_None,
+				FAppStyle::GetAppStyleSetName())
+		{}
+
+		virtual void RegisterCommands() override;
+
+	public:
+		/** Adds a new page to the collection. */
+		TSharedPtr<FUICommandInfo> AddPage;
+
+		/** Copies the given page and adds it to the collection. */
+		TSharedPtr<FUICommandInfo> CopyPage;
+
+		/** Deletes an existing page from the collection. */
+		TSharedPtr<FUICommandInfo> DeletePage;
+
+		/** Renders the available page(s) in batch. */
+		TSharedPtr<FUICommandInfo> BatchRenderList;
+
+		/** Toggles Auto Compilation in the blueprint graph. */
+		TSharedPtr<FUICommandInfo> AutoCompileGraph;
+	};
+}
diff --git a/Engine/Plugins/Experimental/RenderPages/Source/RenderPagesEditor/Private/Factories/RenderPageCollectionFactory.cpp b/Engine/Plugins/Experimental/RenderPages/Source/RenderPagesEditor/Private/Factories/RenderPageCollectionFactory.cpp
new file mode 100644
--- /dev/null
+++ b/Engine/Plugins/Experimental/RenderPages/Source/RenderPagesEditor/Private/Factories/RenderPageCollectionFactory.cpp
@@ -0,0 +1,87 @@
+// Copyright Epic Games, Inc. All Rights Reserved.
+
+#include "Factories/RenderPageCollectionFactory.h"
+#include "Blueprints/RenderPagesBlueprint.h"
+#include "Graph/RenderPagesGraph.h"
+#include "Graph/RenderPagesGraphSchema.h"
+#include "IContentBrowserSingleton.h"
+#include "Kismet2/BlueprintEditorUtils.h"
+#include "Kismet2/KismetEditorUtilities.h"
+#include "Misc/MessageDialog.h"
+#include "RenderPage/RenderPageCollection.h"
+#include "RenderPage/RenderPagesBlueprintGeneratedClass.h"
+
+#define LOCTEXT_NAMESPACE "RenderPagesBlueprintFactory"
+
+
+URenderPagesBlueprintFactory::URenderPagesBlueprintFactory()
+{
+	ParentClass = URenderPageCollection::StaticClass();
+	SupportedClass = URenderPagesBlueprint::StaticClass();
+	bCreateNew = true; // This factory manufacture new objects from scratch
+	bEditAfterNew = true; // This factory will open the editor for each new object
+}
+
+UObject* URenderPagesBlueprintFactory::FactoryCreateNew(UClass* InClass, UObject* InParent, FName InName, EObjectFlags Flags, UObject* Context, FFeedbackContext* Warn, FName CallingContext)
+{
+	// Make sure we are trying to factory a Render Pages Blueprint, then create and init one
+	check(InClass->IsChildOf(URenderPagesBlueprint::StaticClass()));
+
+	if ((ParentClass == nullptr) || !FKismetEditorUtilities::CanCreateBlueprintOfClass(ParentClass) || !ParentClass->IsChildOf(URenderPageCollection::StaticClass()))
+	{
+		FFormatNamedArguments Args;
+		Args.Add(TEXT("ClassName"), (ParentClass != nullptr) ? FText::FromString(ParentClass->GetName()) : LOCTEXT("Null", "(null)"));
+		FMessageDialog::Open(EAppMsgType::Ok, FText::Format(LOCTEXT("CannotCreateRenderPagesBlueprint", "Cannot create an Render Pages Blueprint based on the class '{0}'."), Args));
+		return nullptr;
+	}
+
+	URenderPagesBlueprint* RenderPagesBlueprint = CastChecked<URenderPagesBlueprint>(FKismetEditorUtilities::CreateBlueprint(ParentClass, InParent, InName, BPTYPE_Normal, URenderPagesBlueprint::StaticClass(), URenderPagesBlueprintGeneratedClass::StaticClass(), CallingContext));
+	CreateRenderPagesGraphIfRequired(RenderPagesBlueprint);
+	return RenderPagesBlueprint;
+}
+
+UObject* URenderPagesBlueprintFactory::FactoryCreateNew(UClass* InClass, UObject* InParent, FName InName, EObjectFlags Flags, UObject* Context, FFeedbackContext* Warn)
+{
+	return FactoryCreateNew(InClass, InParent, InName, Flags, Context, Warn, NAME_None);
+}
+
+bool URenderPagesBlueprintFactory::ShouldShowInNewMenu() const
+{
+	return true;
+}
+
+uint32 URenderPagesBlueprintFactory::GetMenuCategories() const
+{
+	//  if wanting to show it in its own category:
+	// IAssetTools& AssetTools = FModuleManager::LoadModuleChecked<FAssetToolsModule>("AssetTools").Get();
+	// return AssetTools.RegisterAdvancedAssetCategory("Render Pages", LOCTEXT("AssetCategoryName", "Render Pages"));
+	
+	return EAssetTypeCategories::Misc;
+}
+
+void URenderPagesBlueprintFactory::CreateRenderPagesGraphIfRequired(URenderPagesBlueprint* InBlueprint)
+{
+	if (InBlueprint == nullptr)
+	{
+		return;
+	}
+
+	for (UEdGraph* EdGraph : InBlueprint->UbergraphPages)
+	{
+		if (EdGraph->IsA<URenderPagesGraph>())
+		{
+			return;
+		}
+	}
+
+	// add an initial graph for us to work in
+	const URenderPagesGraphSchema* RenderPagesGraphSchema = GetDefault<URenderPagesGraphSchema>();
+	UEdGraph* RenderPagesGraph = FBlueprintEditorUtils::CreateNewGraph(InBlueprint, RenderPagesGraphSchema->GraphName_RenderPages, URenderPagesGraph::StaticClass(), URenderPagesGraphSchema::StaticClass());
+	RenderPagesGraph->bAllowDeletion = false;
+	FBlueprintEditorUtils::AddUbergraphPage(InBlueprint, RenderPagesGraph);
+	InBlueprint->LastEditedDocuments.AddUnique(RenderPagesGraph);
+	InBlueprint->PostLoad();
+}
+
+
+#undef LOCTEXT_NAMESPACE
diff --git a/Engine/Plugins/Experimental/RenderPages/Source/RenderPagesEditor/Private/Factories/RenderPageCollectionFactory.h b/Engine/Plugins/Experimental/RenderPages/Source/RenderPagesEditor/Private/Factories/RenderPageCollectionFactory.h
new file mode 100644
--- /dev/null
+++ b/Engine/Plugins/Experimental/RenderPages/Source/RenderPagesEditor/Private/Factories/RenderPageCollectionFactory.h
@@ -0,0 +1,40 @@
+// Copyright Epic Games, Inc. All Rights Reserved.
+
+#pragma once
+
+#include "CoreMinimal.h"
+#include "Factories/Factory.h"
+#include "RenderPage/RenderPagePropsSource.h"
+#include "RenderPageCollectionFactory.generated.h"
+
+
+class URenderPagesBlueprint;
+class URenderPageCollection;
+
+
+/**
+ * The factory that creates URenderPageBlueprint (render page collection) instances.
+ */
+UCLASS(MinimalAPI, HideCategories=Object)
+class URenderPagesBlueprintFactory : public UFactory
+{
+	GENERATED_BODY()
+
+public:
+	URenderPagesBlueprintFactory();
+
+	virtual UObject* FactoryCreateNew(UClass* InClass, UObject* InParent, FName InName, EObjectFlags Flags, UObject* Context, FFeedbackContext* Warn, FName CallingContext) override;
+	virtual UObject* FactoryCreateNew(UClass* InClass, UObject* InParent, FName InName, EObjectFlags Flags, UObject* Context, FFeedbackContext* Warn) override;
+	virtual bool ShouldShowInNewMenu() const override;
+	virtual uint32 GetMenuCategories() const override;
+
+	/**
+	 * Creates the render page graph if necessary.
+	 * @param InBlueprint The blueprint to setup up
+	 */
+	static void CreateRenderPagesGraphIfRequired(URenderPagesBlueprint* InBlueprint);
+
+	/** The parent class of the created blueprint. */
+	UPROPERTY(EditAnywhere, Category="Render Pages|Render Pages Factory", Meta = (AllowAbstract = ""))
+	TSubclassOf<URenderPageCollection> ParentClass;
+};
diff --git a/Engine/Plugins/Experimental/RenderPages/Source/RenderPagesEditor/Private/Factories/RenderPagePropsSourceWidgetFactoryLocal.cpp b/Engine/Plugins/Experimental/RenderPages/Source/RenderPagesEditor/Private/Factories/RenderPagePropsSourceWidgetFactoryLocal.cpp
new file mode 100644
--- /dev/null
+++ b/Engine/Plugins/Experimental/RenderPages/Source/RenderPagesEditor/Private/Factories/RenderPagePropsSourceWidgetFactoryLocal.cpp
@@ -0,0 +1,17 @@
+// Copyright Epic Games, Inc. All Rights Reserved.
+
+#pragma once
+
+#include "Factories/RenderPagePropsSourceWidgetFactoryLocal.h"
+#include "UI/SRenderPagesPropsLocal.h"
+#include "RenderPage/RenderPagePropsSource.h"
+
+
+TSharedPtr<UE::RenderPages::Private::SRenderPagesPropsBase> UE::RenderPages::Private::FRenderPagePropsSourceWidgetFactoryLocal::CreateInstance(URenderPagePropsSourceBase* PropsSource, TSharedPtr<IRenderPageCollectionEditor> BlueprintEditor)
+{
+	if (URenderPagePropsSourceLocal* InPropsSource = Cast<URenderPagePropsSourceLocal>(PropsSource))
+	{
+		return SNew(SRenderPagesPropsLocal, BlueprintEditor, InPropsSource);
+	}
+	return nullptr;
+}
diff --git a/Engine/Plugins/Experimental/RenderPages/Source/RenderPagesEditor/Private/Factories/RenderPagePropsSourceWidgetFactoryLocal.h b/Engine/Plugins/Experimental/RenderPages/Source/RenderPagesEditor/Private/Factories/RenderPagePropsSourceWidgetFactoryLocal.h
new file mode 100644
--- /dev/null
+++ b/Engine/Plugins/Experimental/RenderPages/Source/RenderPagesEditor/Private/Factories/RenderPagePropsSourceWidgetFactoryLocal.h
@@ -0,0 +1,19 @@
+// Copyright Epic Games, Inc. All Rights Reserved.
+
+#pragma once
+
+#include "CoreMinimal.h"
+#include "Factories/IRenderPagePropsSourceWidgetFactory.h"
+
+
+namespace UE::RenderPages::Private
+{
+	/**
+	 * The factory that creates props source instances for the props type "local".
+	 */
+	class RENDERPAGESEDITOR_API FRenderPagePropsSourceWidgetFactoryLocal final : public IRenderPagePropsSourceWidgetFactory
+	{
+	public:
+		virtual TSharedPtr<SRenderPagesPropsBase> CreateInstance(URenderPagePropsSourceBase* PropsSource, TSharedPtr<IRenderPageCollectionEditor> BlueprintEditor) override;
+	};
+}
diff --git a/Engine/Plugins/Experimental/RenderPages/Source/RenderPagesEditor/Private/Factories/RenderPagePropsSourceWidgetFactoryRemoteControl.cpp b/Engine/Plugins/Experimental/RenderPages/Source/RenderPagesEditor/Private/Factories/RenderPagePropsSourceWidgetFactoryRemoteControl.cpp
new file mode 100644
--- /dev/null
+++ b/Engine/Plugins/Experimental/RenderPages/Source/RenderPagesEditor/Private/Factories/RenderPagePropsSourceWidgetFactoryRemoteControl.cpp
@@ -0,0 +1,17 @@
+// Copyright Epic Games, Inc. All Rights Reserved.
+
+#pragma once
+
+#include "Factories/RenderPagePropsSourceWidgetFactoryRemoteControl.h"
+#include "UI/SRenderPagesPropsRemoteControl.h"
+#include "RenderPage/RenderPagePropsSource.h"
+
+
+TSharedPtr<UE::RenderPages::Private::SRenderPagesPropsBase> UE::RenderPages::Private::FRenderPagePropsSourceWidgetFactoryRemoteControl::CreateInstance(URenderPagePropsSourceBase* PropsSource, TSharedPtr<IRenderPageCollectionEditor> BlueprintEditor)
+{
+	if (URenderPagePropsSourceRemoteControl* InPropsSource = Cast<URenderPagePropsSourceRemoteControl>(PropsSource))
+	{
+		return SNew(SRenderPagesPropsRemoteControl, BlueprintEditor, InPropsSource);
+	}
+	return nullptr;
+}
diff --git a/Engine/Plugins/Experimental/RenderPages/Source/RenderPagesEditor/Private/Factories/RenderPagePropsSourceWidgetFactoryRemoteControl.h b/Engine/Plugins/Experimental/RenderPages/Source/RenderPagesEditor/Private/Factories/RenderPagePropsSourceWidgetFactoryRemoteControl.h
new file mode 100644
--- /dev/null
+++ b/Engine/Plugins/Experimental/RenderPages/Source/RenderPagesEditor/Private/Factories/RenderPagePropsSourceWidgetFactoryRemoteControl.h
@@ -0,0 +1,19 @@
+// Copyright Epic Games, Inc. All Rights Reserved.
+
+#pragma once
+
+#include "CoreMinimal.h"
+#include "Factories/IRenderPagePropsSourceWidgetFactory.h"
+
+
+namespace UE::RenderPages::Private
+{
+	/**
+	 * The factory that creates props source instances for the props type "remote control".
+	 */
+	class RENDERPAGESEDITOR_API FRenderPagePropsSourceWidgetFactoryRemoteControl final : public IRenderPagePropsSourceWidgetFactory
+	{
+	public:
+		virtual TSharedPtr<SRenderPagesPropsBase> CreateInstance(URenderPagePropsSourceBase* PropsSource, TSharedPtr<IRenderPageCollectionEditor> BlueprintEditor) override;
+	};
+}
diff --git a/Engine/Plugins/Experimental/RenderPages/Source/RenderPagesEditor/Private/RenderPagesEditorModule.cpp b/Engine/Plugins/Experimental/RenderPages/Source/RenderPagesEditor/Private/RenderPagesEditorModule.cpp
new file mode 100644
--- /dev/null
+++ b/Engine/Plugins/Experimental/RenderPages/Source/RenderPagesEditor/Private/RenderPagesEditorModule.cpp
@@ -0,0 +1,128 @@
+// Copyright Epic Games, Inc. All Rights Reserved.
+
+#include "RenderPagesEditorModule.h"
+
+#include "AssetTypeActions/RenderPagesBlueprintActions.h"
+#include "Blueprints/RenderPagesBlueprint.h"
+#include "Commands/RenderPagesEditorCommands.h"
+#include "Factories/RenderPageCollectionFactory.h"
+#include "Factories/RenderPagePropsSourceWidgetFactoryLocal.h"
+#include "Factories/RenderPagePropsSourceWidgetFactoryRemoteControl.h"
+#include "RenderPage/RenderPageCollection.h"
+#include "Styles/RenderPagesEditorStyle.h"
+#include "Toolkit/RenderPageCollectionEditor.h"
+
+#include "IAssetTools.h"
+#include "Kismet2/KismetEditorUtilities.h"
+
+#define LOCTEXT_NAMESPACE "RenderPagesEditorModule"
+
+
+void UE::RenderPages::Private::FRenderPagesEditorModule::StartupModule()
+{
+	FRenderPagesEditorStyle::Initialize();
+	FRenderPagesEditorStyle::ReloadTextures();
+	FRenderPagesEditorCommands::Register();
+
+	MenuExtensibilityManager = MakeShared<FExtensibilityManager>();
+	ToolBarExtensibilityManager = MakeShared<FExtensibilityManager>();
+
+	RegisterPropsSourceWidgetFactories();
+
+	// Register asset tools
+	auto RegisterAssetTypeAction = [this](const TSharedRef<IAssetTypeActions>& InAssetTypeAction)
+	{
+		IAssetTools& AssetTools = FModuleManager::LoadModuleChecked<FAssetToolsModule>("AssetTools").Get();
+		RegisteredAssetTypeActions.Add(InAssetTypeAction);
+		AssetTools.RegisterAssetTypeActions(InAssetTypeAction);
+	};
+
+	RegisterAssetTypeAction(MakeShared<FRenderPagesBlueprintActions>());
+
+	// Register to fixup newly created BPs
+	FKismetEditorUtilities::RegisterOnBlueprintCreatedCallback(this, URenderPageCollection::StaticClass(), FKismetEditorUtilities::FOnBlueprintCreated::CreateRaw(this, &FRenderPagesEditorModule::HandleNewBlueprintCreated));
+}
+
+void UE::RenderPages::Private::FRenderPagesEditorModule::ShutdownModule()
+{
+	FAssetToolsModule* AssetToolsModule = FModuleManager::GetModulePtr<FAssetToolsModule>("AssetTools");
+
+	if (AssetToolsModule)
+	{
+		for (TSharedRef<IAssetTypeActions> RegisteredAssetTypeAction : RegisteredAssetTypeActions)
+		{
+			AssetToolsModule->Get().UnregisterAssetTypeActions(RegisteredAssetTypeAction);
+		}
+	}
+
+	UnregisterPropsSourceWidgetFactories();
+
+	MenuExtensibilityManager.Reset();
+	ToolBarExtensibilityManager.Reset();
+
+	FRenderPagesEditorCommands::Unregister();
+	FRenderPagesEditorStyle::Shutdown();
+}
+
+void UE::RenderPages::Private::FRenderPagesEditorModule::RegisterPropsSourceWidgetFactories()
+{
+	RegisterPropsSourceWidgetFactory(ERenderPagePropsSourceType::Local, MakeShared<FRenderPagePropsSourceWidgetFactoryLocal>());
+	RegisterPropsSourceWidgetFactory(ERenderPagePropsSourceType::RemoteControl, MakeShared<FRenderPagePropsSourceWidgetFactoryRemoteControl>());
+}
+
+void UE::RenderPages::Private::FRenderPagesEditorModule::UnregisterPropsSourceWidgetFactories()
+{
+	UnregisterPropsSourceWidgetFactory(ERenderPagePropsSourceType::Local);
+	UnregisterPropsSourceWidgetFactory(ERenderPagePropsSourceType::RemoteControl);
+}
+
+void UE::RenderPages::Private::FRenderPagesEditorModule::RegisterPropsSourceWidgetFactory(const ERenderPagePropsSourceType PropsSourceType, const TSharedPtr<IRenderPagePropsSourceWidgetFactory>& InFactory)
+{
+	PropsSourceWidgetFactories.Add(PropsSourceType, InFactory);
+}
+
+void UE::RenderPages::Private::FRenderPagesEditorModule::UnregisterPropsSourceWidgetFactory(const ERenderPagePropsSourceType PropsSourceType)
+{
+	PropsSourceWidgetFactories.Remove(PropsSourceType);
+}
+
+TSharedPtr<UE::RenderPages::Private::SRenderPagesPropsBase> UE::RenderPages::Private::FRenderPagesEditorModule::CreatePropsSourceWidget(URenderPagePropsSourceBase* PropsSource, TSharedPtr<IRenderPageCollectionEditor> BlueprintEditor)
+{
+	if (!PropsSource)
+	{
+		return nullptr;
+	}
+
+	TSharedPtr<IRenderPagePropsSourceWidgetFactory>* FactoryPtr = PropsSourceWidgetFactories.Find(PropsSource->GetType());
+	if (!FactoryPtr)
+	{
+		return nullptr;
+	}
+
+	TSharedPtr<IRenderPagePropsSourceWidgetFactory> Factory = *FactoryPtr;
+	if (!Factory)
+	{
+		return nullptr;
+	}
+
+	return Factory->CreateInstance(PropsSource, BlueprintEditor);
+}
+
+TSharedRef<UE::RenderPages::IRenderPageCollectionEditor> UE::RenderPages::Private::FRenderPagesEditorModule::CreateRenderPageCollectionEditor(const EToolkitMode::Type Mode, const TSharedPtr<IToolkitHost>& InitToolkitHost, URenderPagesBlueprint* InBlueprint)
+{
+	TSharedRef<FRenderPageCollectionEditor> NewRenderPagesEditor = MakeShared<FRenderPageCollectionEditor>();
+
+	NewRenderPagesEditor->InitRenderPagesEditor(Mode, InitToolkitHost, InBlueprint);
+
+	return NewRenderPagesEditor;
+}
+
+void UE::RenderPages::Private::FRenderPagesEditorModule::HandleNewBlueprintCreated(UBlueprint* InBlueprint)
+{
+	URenderPagesBlueprintFactory::CreateRenderPagesGraphIfRequired(Cast<URenderPagesBlueprint>(InBlueprint));
+}
+
+
+#undef LOCTEXT_NAMESPACE
+
+IMPLEMENT_MODULE(UE::RenderPages::Private::FRenderPagesEditorModule, RenderPagesEditor)
diff --git a/Engine/Plugins/Experimental/RenderPages/Source/RenderPagesEditor/Private/RenderPagesEditorModule.h b/Engine/Plugins/Experimental/RenderPages/Source/RenderPagesEditor/Private/RenderPagesEditorModule.h
new file mode 100644
--- /dev/null
+++ b/Engine/Plugins/Experimental/RenderPages/Source/RenderPagesEditor/Private/RenderPagesEditorModule.h
@@ -0,0 +1,60 @@
+// Copyright Epic Games, Inc. All Rights Reserved.
+
+#pragma once
+
+#include "CoreMinimal.h"
+#include "IRenderPagesEditorModule.h"
+
+
+class IAssetTypeActions;
+class IToolkitHost;
+class URenderPageCollection;
+class URenderPagesBlueprint;
+
+namespace UE::RenderPages::Private
+{
+	class FRenderPageCollectionEditor;
+}
+
+
+namespace UE::RenderPages::Private
+{
+	class FRenderPagesEditorModule : public IRenderPagesEditorModule
+	{
+	public:
+		//~ Begin IModuleInterface interface
+		virtual void StartupModule() override;
+		virtual void ShutdownModule() override;
+		//~ End IModuleInterface interface
+
+		//~ Begin IRenderPagesEditorModule interface
+		virtual TSharedRef<IRenderPageCollectionEditor> CreateRenderPageCollectionEditor(const EToolkitMode::Type Mode, const TSharedPtr<IToolkitHost>& InitToolkitHost, URenderPagesBlueprint* InBlueprint) override;
+
+		virtual TSharedPtr<SRenderPagesPropsBase> CreatePropsSourceWidget(URenderPagePropsSourceBase* PropsSource, TSharedPtr<IRenderPageCollectionEditor> BlueprintEditor) override;
+		virtual const TMap<ERenderPagePropsSourceType, TSharedPtr<IRenderPagePropsSourceWidgetFactory>>& GetPropsSourceWidgetFactories() const override { return PropsSourceWidgetFactories; }
+		//~ End IRenderPagesEditorModule interface
+
+		/** Gets the extensibility managers for outside entities to extend gui page editor's menus. */
+		virtual TSharedPtr<FExtensibilityManager> GetMenuExtensibilityManager() override { return MenuExtensibilityManager; }
+
+		/** Gets the extensibility managers for outside entities to extend gui page editor's and toolbars. */
+		virtual TSharedPtr<FExtensibilityManager> GetToolBarExtensibilityManager() override { return ToolBarExtensibilityManager; }
+
+	private:
+		/** Handle a new render pages blueprint being created. */
+		void HandleNewBlueprintCreated(UBlueprint* InBlueprint);
+
+		void RegisterPropsSourceWidgetFactories();
+		void UnregisterPropsSourceWidgetFactories();
+		void RegisterPropsSourceWidgetFactory(const ERenderPagePropsSourceType PropsSourceType, const TSharedPtr<IRenderPagePropsSourceWidgetFactory>& InFactory);
+		void UnregisterPropsSourceWidgetFactory(const ERenderPagePropsSourceType PropsSourceType);
+
+	private:
+		TSharedPtr<FExtensibilityManager> MenuExtensibilityManager;
+		TSharedPtr<FExtensibilityManager> ToolBarExtensibilityManager;
+
+		TArray<TSharedRef<IAssetTypeActions>> RegisteredAssetTypeActions;
+
+		TMap<ERenderPagePropsSourceType, TSharedPtr<IRenderPagePropsSourceWidgetFactory>> PropsSourceWidgetFactories;
+	};
+}
diff --git a/Engine/Plugins/Experimental/RenderPages/Source/RenderPagesEditor/Private/Styles/RenderPagesEditorStyle.cpp b/Engine/Plugins/Experimental/RenderPages/Source/RenderPagesEditor/Private/Styles/RenderPagesEditorStyle.cpp
new file mode 100644
--- /dev/null
+++ b/Engine/Plugins/Experimental/RenderPages/Source/RenderPagesEditor/Private/Styles/RenderPagesEditorStyle.cpp
@@ -0,0 +1,117 @@
+// Copyright Epic Games, Inc. All Rights Reserved.
+
+#include "Styles/RenderPagesEditorStyle.h"
+
+#include "Brushes/SlateBorderBrush.h"
+#include "Brushes/SlateImageBrush.h"
+#include "Framework/Application/SlateApplication.h"
+#include "Interfaces/IPluginManager.h"
+#include "Styling/SlateStyle.h"
+#include "Styling/SlateStyleRegistry.h"
+#include "Styling/StyleColors.h"
+
+
+TSharedPtr<FSlateStyleSet> UE::RenderPages::Private::FRenderPagesEditorStyle::StyleInstance = nullptr;
+
+
+const ISlateStyle& UE::RenderPages::Private::FRenderPagesEditorStyle::Get()
+{
+	return *StyleInstance;
+}
+
+void UE::RenderPages::Private::FRenderPagesEditorStyle::Initialize()
+{
+	if (!StyleInstance.IsValid())
+	{
+		StyleInstance = Create();
+		FSlateStyleRegistry::RegisterSlateStyle(*StyleInstance);
+	}
+}
+
+void UE::RenderPages::Private::FRenderPagesEditorStyle::Shutdown()
+{
+	FSlateStyleRegistry::UnRegisterSlateStyle(*StyleInstance);
+	ensure(StyleInstance.IsUnique());
+	StyleInstance.Reset();
+}
+
+FName UE::RenderPages::Private::FRenderPagesEditorStyle::GetStyleSetName()
+{
+	static FName StyleSetName(TEXT("LevelSnapshotsEditor"));
+	return StyleSetName;
+}
+
+const FLinearColor& UE::RenderPages::Private::FRenderPagesEditorStyle::GetColor(FName PropertyName, const ANSICHAR* Specifier)
+{
+	return StyleInstance->GetColor(PropertyName, Specifier);
+}
+
+
+const FSlateBrush* UE::RenderPages::Private::FRenderPagesEditorStyle::GetBrush(FName PropertyName, const ANSICHAR* Specifier)
+{
+	return StyleInstance->GetBrush(PropertyName, Specifier);
+}
+
+#define IMAGE_BRUSH(RelativePath, ...) FSlateImageBrush(Style->RootToContentDir(RelativePath, TEXT(".png")), __VA_ARGS__)
+#define BOX_BRUSH(RelativePath, ...) FSlateBoxBrush(Style->RootToContentDir(RelativePath, TEXT(".png")), __VA_ARGS__)
+#define BORDER_BRUSH(RelativePath, ...) FSlateBorderBrush(Style->RootToContentDir(RelativePath, TEXT(".png")), __VA_ARGS__)
+#define IMAGE_BRUSH_SVG(RelativePath, ...) FSlateVectorImageBrush(Style->RootToContentDir(RelativePath, TEXT(".svg")), __VA_ARGS__)
+
+const FVector2D Icon64x64(64.f, 64.f);
+const FVector2D Icon40x40(40.0f, 40.0f);
+const FVector2D Icon20x20(20.0f, 20.0f);
+const FVector2D Icon16x16(16.0f, 16.0f);
+const FVector2D Icon12x12(12.0f, 12.0f);
+
+TSharedRef<FSlateStyleSet> UE::RenderPages::Private::FRenderPagesEditorStyle::Create()
+{
+	TSharedRef<FSlateStyleSet> Style = MakeShared<FSlateStyleSet>("RenderPagesEditor");
+
+	TSharedPtr<IPlugin> Plugin = IPluginManager::Get().FindPlugin(TEXT("RenderPages"));
+	check(Plugin.IsValid());
+	if (Plugin.IsValid())
+	{
+		Style->SetContentRoot(FPaths::Combine(Plugin->GetBaseDir(), TEXT("Resources")));
+	}
+
+	Style->Set("RenderPages.AutoCompileGraph", new IMAGE_BRUSH("Slate/AutoCompile", Icon20x20));
+	Style->Set("RenderPages.AutoCompileGraph.Small", new IMAGE_BRUSH("Slate/AutoCompile", Icon20x20));
+
+	Style->Set("Invisible",
+		FButtonStyle()
+		.SetNormal(FSlateNoResource())
+		.SetHovered(FSlateNoResource())
+		.SetPressed(FSlateNoResource())
+		.SetNormalPadding(FMargin())
+		.SetPressedPadding(FMargin())
+	);
+	Style->Set("HoverHintOnly",
+		FButtonStyle()
+		.SetNormal(FSlateNoResource())
+		.SetHovered(FSlateNoResource())
+		.SetPressed(FSlateNoResource())
+		.SetHovered(BOX_BRUSH("Common/ButtonHoverHint", FMargin(4/16.0f), FLinearColor(1,1,1,0.05f)))
+		.SetPressed(BOX_BRUSH("Common/ButtonHoverHint", FMargin(4/16.0f), FLinearColor(1,1,1,0.15f)))
+		.SetNormalPadding(FMargin())
+		.SetPressedPadding(FMargin())
+	);
+
+	Style->Set("TabButton",
+		FButtonStyle(FAppStyle::Get().GetWidgetStyle<FButtonStyle>("Button"))
+		.SetDisabled(BOX_BRUSH("Common/FlatButton", 2.0f / 8.0f, FAppStyle::GetSlateColor("SelectionColor")))
+		.SetDisabledForeground(FStyleColors::White)
+	);
+
+
+	return Style;
+}
+
+#undef IMAGE_BRUSH
+#undef BOX_BRUSH
+#undef BORDER_BRUSH
+#undef IMAGE_BRUSH_SVG
+
+void UE::RenderPages::Private::FRenderPagesEditorStyle::ReloadTextures()
+{
+	FSlateApplication::Get().GetRenderer()->ReloadTextureResources();
+}
diff --git a/Engine/Plugins/Experimental/RenderPages/Source/RenderPagesEditor/Private/Styles/RenderPagesEditorStyle.h b/Engine/Plugins/Experimental/RenderPages/Source/RenderPagesEditor/Private/Styles/RenderPagesEditorStyle.h
new file mode 100644
--- /dev/null
+++ b/Engine/Plugins/Experimental/RenderPages/Source/RenderPagesEditor/Private/Styles/RenderPagesEditorStyle.h
@@ -0,0 +1,37 @@
+// Copyright Epic Games, Inc. All Rights Reserved.
+
+#pragma once
+
+#include "Styling/ISlateStyle.h"
+#include "Styling/SlateStyle.h"
+
+
+namespace UE::RenderPages::Private
+{
+	/**
+	 * A static class containing the style setup of the render pages plugin.
+	 */
+	class FRenderPagesEditorStyle
+	{
+	public:
+		static const ISlateStyle& Get();
+
+		static void Initialize();
+		static void Shutdown();
+		static void ReloadTextures();
+
+		static FName GetStyleSetName();
+		static const FLinearColor& GetColor(FName PropertyName, const ANSICHAR* Specifier = nullptr);
+		static const FSlateBrush* GetBrush(FName PropertyName, const ANSICHAR* Specifier = nullptr);
+
+		template<typename WidgetStyleType>
+		static const WidgetStyleType& GetWidgetStyle(FName PropertyName, const ANSICHAR* Specifier = nullptr)
+		{
+			return StyleInstance->GetWidgetStyle<WidgetStyleType>(PropertyName, Specifier);
+		}
+
+	private:
+		static TSharedRef<FSlateStyleSet> Create();
+		static TSharedPtr<FSlateStyleSet> StyleInstance;
+	};
+}
diff --git a/Engine/Plugins/Experimental/RenderPages/Source/RenderPagesEditor/Private/TabFactory/CollectionPropertiesTabSummoner.cpp b/Engine/Plugins/Experimental/RenderPages/Source/RenderPagesEditor/Private/TabFactory/CollectionPropertiesTabSummoner.cpp
new file mode 100644
--- /dev/null
+++ b/Engine/Plugins/Experimental/RenderPages/Source/RenderPagesEditor/Private/TabFactory/CollectionPropertiesTabSummoner.cpp
@@ -0,0 +1,32 @@
+// Copyright Epic Games, Inc. All Rights Reserved.
+
+#include "TabFactory/CollectionPropertiesTabSummoner.h"
+#include "UI/Tabs/SRenderPagesCollectionPropertiesTab.h"
+#include "Styling/AppStyle.h"
+#include "IRenderPageCollectionEditor.h"
+
+#define LOCTEXT_NAMESPACE "RenderPages"
+
+
+const FName UE::RenderPages::Private::FCollectionPropertiesTabSummoner::TabID(TEXT("CollectionProperties"));
+
+
+UE::RenderPages::Private::FCollectionPropertiesTabSummoner::FCollectionPropertiesTabSummoner(TSharedPtr<IRenderPageCollectionEditor> InBlueprintEditor)
+	: FWorkflowTabFactory(TabID, InBlueprintEditor)
+	, BlueprintEditorWeakPtr(InBlueprintEditor)
+{
+	TabLabel = LOCTEXT("CollectionProperties_TabLabel", "Collection");
+	TabIcon = FSlateIcon(FAppStyle::GetAppStyleSetName(), "ContentBrowser.ShowSourcesView");
+	bIsSingleton = true;
+	ViewMenuDescription = LOCTEXT("CollectionProperties_ViewMenu_Desc", "Collection");
+	ViewMenuTooltip = LOCTEXT("CollectionProperties_ViewMenu_ToolTip", "Show the page collection properties.");
+}
+
+TSharedRef<SWidget> UE::RenderPages::Private::FCollectionPropertiesTabSummoner::CreateTabBody(const FWorkflowTabSpawnInfo& Info) const
+{
+	return SNew(SRenderPagesCollectionPropertiesTab, BlueprintEditorWeakPtr.Pin())
+		.AddMetaData<FTagMetaData>(FTagMetaData(TEXT("Collection")));
+}
+
+
+#undef LOCTEXT_NAMESPACE
diff --git a/Engine/Plugins/Experimental/RenderPages/Source/RenderPagesEditor/Private/TabFactory/CollectionPropertiesTabSummoner.h b/Engine/Plugins/Experimental/RenderPages/Source/RenderPagesEditor/Private/TabFactory/CollectionPropertiesTabSummoner.h
new file mode 100644
--- /dev/null
+++ b/Engine/Plugins/Experimental/RenderPages/Source/RenderPagesEditor/Private/TabFactory/CollectionPropertiesTabSummoner.h
@@ -0,0 +1,34 @@
+// Copyright Epic Games, Inc. All Rights Reserved.
+
+#pragma once
+
+#include "CoreMinimal.h"
+#include "WorkflowOrientedApp/WorkflowTabFactory.h"
+
+
+namespace UE::RenderPages
+{
+	class IRenderPageCollectionEditor;
+}
+
+
+namespace UE::RenderPages::Private
+{
+	/**
+	 * The collection properties tab factory.
+	 */
+	struct FCollectionPropertiesTabSummoner : FWorkflowTabFactory
+	{
+	public:
+		/** Unique ID representing this tab. */
+		static const FName TabID;
+
+	public:
+		FCollectionPropertiesTabSummoner(TSharedPtr<IRenderPageCollectionEditor> InBlueprintEditor);
+		virtual TSharedRef<SWidget> CreateTabBody(const FWorkflowTabSpawnInfo& Info) const override;
+
+	protected:
+		/** A weak reference to the blueprint editor. */
+		TWeakPtr<IRenderPageCollectionEditor> BlueprintEditorWeakPtr;
+	};
+}
diff --git a/Engine/Plugins/Experimental/RenderPages/Source/RenderPagesEditor/Private/TabFactory/PageListTabSummoner.cpp b/Engine/Plugins/Experimental/RenderPages/Source/RenderPagesEditor/Private/TabFactory/PageListTabSummoner.cpp
new file mode 100644
--- /dev/null
+++ b/Engine/Plugins/Experimental/RenderPages/Source/RenderPagesEditor/Private/TabFactory/PageListTabSummoner.cpp
@@ -0,0 +1,32 @@
+// Copyright Epic Games, Inc. All Rights Reserved.
+
+#include "TabFactory/PageListTabSummoner.h"
+#include "UI/Tabs/SRenderPagesPageListTab.h"
+#include "Styling/AppStyle.h"
+#include "IRenderPageCollectionEditor.h"
+
+#define LOCTEXT_NAMESPACE "RenderPages"
+
+
+const FName UE::RenderPages::Private::FPageListTabSummoner::TabID(TEXT("PageList"));
+
+
+UE::RenderPages::Private::FPageListTabSummoner::FPageListTabSummoner(TSharedPtr<IRenderPageCollectionEditor> InBlueprintEditor)
+	: FWorkflowTabFactory(TabID, InBlueprintEditor)
+	, BlueprintEditorWeakPtr(InBlueprintEditor)
+{
+	TabLabel = LOCTEXT("PageList_TabLabel", "Pages");
+	TabIcon = FSlateIcon(FAppStyle::GetAppStyleSetName(), "PlacementBrowser.Icons.All");
+	bIsSingleton = true;
+	ViewMenuDescription = LOCTEXT("PageList_ViewMenu_Desc", "Pages");
+	ViewMenuTooltip = LOCTEXT("PageList_ViewMenu_ToolTip", "Show the page list.");
+}
+
+TSharedRef<SWidget> UE::RenderPages::Private::FPageListTabSummoner::CreateTabBody(const FWorkflowTabSpawnInfo& Info) const
+{
+	return SNew(SRenderPagesPageListTab, BlueprintEditorWeakPtr.Pin())
+		.AddMetaData<FTagMetaData>(FTagMetaData(TEXT("Page List")));
+}
+
+
+#undef LOCTEXT_NAMESPACE
diff --git a/Engine/Plugins/Experimental/RenderPages/Source/RenderPagesEditor/Private/TabFactory/PageListTabSummoner.h b/Engine/Plugins/Experimental/RenderPages/Source/RenderPagesEditor/Private/TabFactory/PageListTabSummoner.h
new file mode 100644
--- /dev/null
+++ b/Engine/Plugins/Experimental/RenderPages/Source/RenderPagesEditor/Private/TabFactory/PageListTabSummoner.h
@@ -0,0 +1,34 @@
+// Copyright Epic Games, Inc. All Rights Reserved.
+
+#pragma once
+
+#include "CoreMinimal.h"
+#include "WorkflowOrientedApp/WorkflowTabFactory.h"
+
+
+namespace UE::RenderPages
+{
+	class IRenderPageCollectionEditor;
+}
+
+
+namespace UE::RenderPages::Private
+{
+	/**
+	 * The page list tab factory.
+	 */
+	struct FPageListTabSummoner : FWorkflowTabFactory
+	{
+	public:
+		/** Unique ID representing this tab. */
+		static const FName TabID;
+
+	public:
+		FPageListTabSummoner(TSharedPtr<IRenderPageCollectionEditor> InBlueprintEditor);
+		virtual TSharedRef<SWidget> CreateTabBody(const FWorkflowTabSpawnInfo& Info) const override;
+
+	protected:
+		/** A weak reference to the blueprint editor. */
+		TWeakPtr<IRenderPageCollectionEditor> BlueprintEditorWeakPtr;
+	};
+}
diff --git a/Engine/Plugins/Experimental/RenderPages/Source/RenderPagesEditor/Private/TabFactory/PagePropertiesTabSummoner.cpp b/Engine/Plugins/Experimental/RenderPages/Source/RenderPagesEditor/Private/TabFactory/PagePropertiesTabSummoner.cpp
new file mode 100644
--- /dev/null
+++ b/Engine/Plugins/Experimental/RenderPages/Source/RenderPagesEditor/Private/TabFactory/PagePropertiesTabSummoner.cpp
@@ -0,0 +1,32 @@
+// Copyright Epic Games, Inc. All Rights Reserved.
+
+#include "TabFactory/PagePropertiesTabSummoner.h"
+#include "UI/Tabs/SRenderPagesPagePropertiesTab.h"
+#include "Styling/AppStyle.h"
+#include "IRenderPageCollectionEditor.h"
+
+#define LOCTEXT_NAMESPACE "RenderPages"
+
+
+const FName UE::RenderPages::Private::FPagePropertiesTabSummoner::TabID(TEXT("PageProperties"));
+
+
+UE::RenderPages::Private::FPagePropertiesTabSummoner::FPagePropertiesTabSummoner(TSharedPtr<IRenderPageCollectionEditor> InBlueprintEditor)
+	: FWorkflowTabFactory(TabID, InBlueprintEditor)
+	, BlueprintEditorWeakPtr(InBlueprintEditor)
+{
+	TabLabel = LOCTEXT("PageProperties_TabLabel", "Page");
+	TabIcon = FSlateIcon(FAppStyle::GetAppStyleSetName(), "EditorPreferences.TabIcon");
+	bIsSingleton = true;
+	ViewMenuDescription = LOCTEXT("PageProperties_ViewMenu_Desc", "Page");
+	ViewMenuTooltip = LOCTEXT("PageProperties_ViewMenu_ToolTip", "Show the page properties.");
+}
+
+TSharedRef<SWidget> UE::RenderPages::Private::FPagePropertiesTabSummoner::CreateTabBody(const FWorkflowTabSpawnInfo& Info) const
+{
+	return SNew(SRenderPagesPagePropertiesTab, BlueprintEditorWeakPtr.Pin())
+		.AddMetaData<FTagMetaData>(FTagMetaData(TEXT("Page")));
+}
+
+
+#undef LOCTEXT_NAMESPACE
diff --git a/Engine/Plugins/Experimental/RenderPages/Source/RenderPagesEditor/Private/TabFactory/PagePropertiesTabSummoner.h b/Engine/Plugins/Experimental/RenderPages/Source/RenderPagesEditor/Private/TabFactory/PagePropertiesTabSummoner.h
new file mode 100644
--- /dev/null
+++ b/Engine/Plugins/Experimental/RenderPages/Source/RenderPagesEditor/Private/TabFactory/PagePropertiesTabSummoner.h
@@ -0,0 +1,34 @@
+// Copyright Epic Games, Inc. All Rights Reserved.
+
+#pragma once
+
+#include "CoreMinimal.h"
+#include "WorkflowOrientedApp/WorkflowTabFactory.h"
+
+
+namespace UE::RenderPages
+{
+	class IRenderPageCollectionEditor;
+}
+
+
+namespace UE::RenderPages::Private
+{
+	/**
+	 * The page properties tab factory.
+	 */
+	struct FPagePropertiesTabSummoner : FWorkflowTabFactory
+	{
+	public:
+		/** Unique ID representing this tab. */
+		static const FName TabID;
+
+	public:
+		FPagePropertiesTabSummoner(TSharedPtr<IRenderPageCollectionEditor> InBlueprintEditor);
+		virtual TSharedRef<SWidget> CreateTabBody(const FWorkflowTabSpawnInfo& Info) const override;
+
+	protected:
+		/** A weak reference to the blueprint editor. */
+		TWeakPtr<IRenderPageCollectionEditor> BlueprintEditorWeakPtr;
+	};
+}
