First usable version of the render pages plugin.

diff --git a/Engine/Plugins/Experimental/RenderPages/Source/RenderPagesEditor/Private/TabFactory/PageViewerTabSummoner.cpp b/Engine/Plugins/Experimental/RenderPages/Source/RenderPagesEditor/Private/TabFactory/PageViewerTabSummoner.cpp
new file mode 100644
--- /dev/null
+++ b/Engine/Plugins/Experimental/RenderPages/Source/RenderPagesEditor/Private/TabFactory/PageViewerTabSummoner.cpp
@@ -0,0 +1,32 @@
+// Copyright Epic Games, Inc. All Rights Reserved.
+
+#include "TabFactory/PageViewerTabSummoner.h"
+#include "UI/Tabs/SRenderPagesPageViewerTab.h"
+#include "Styling/AppStyle.h"
+#include "IRenderPageCollectionEditor.h"
+
+#define LOCTEXT_NAMESPACE "RenderPages"
+
+
+const FName UE::RenderPages::Private::FPageViewerTabSummoner::TabID(TEXT("PageViewer"));
+
+
+UE::RenderPages::Private::FPageViewerTabSummoner::FPageViewerTabSummoner(TSharedPtr<IRenderPageCollectionEditor> InBlueprintEditor)
+	: FWorkflowTabFactory(TabID, InBlueprintEditor)
+	, BlueprintEditorWeakPtr(InBlueprintEditor)
+{
+	TabLabel = LOCTEXT("PageViewer_TabLabel", "Viewer");
+	TabIcon = FSlateIcon(FAppStyle::GetAppStyleSetName(), "LevelEditor.Tabs.Details");
+	bIsSingleton = true;
+	ViewMenuDescription = LOCTEXT("PageViewer_ViewMenu_Desc", "Viewer");
+	ViewMenuTooltip = LOCTEXT("PageViewer_ViewMenu_ToolTip", "Show the page viewer.");
+}
+
+TSharedRef<SWidget> UE::RenderPages::Private::FPageViewerTabSummoner::CreateTabBody(const FWorkflowTabSpawnInfo& Info) const
+{
+	return SNew(SRenderPagesPageViewerTab, BlueprintEditorWeakPtr.Pin())
+		.AddMetaData<FTagMetaData>(FTagMetaData(TEXT("Viewer")));
+}
+
+
+#undef LOCTEXT_NAMESPACE
diff --git a/Engine/Plugins/Experimental/RenderPages/Source/RenderPagesEditor/Private/TabFactory/PageViewerTabSummoner.h b/Engine/Plugins/Experimental/RenderPages/Source/RenderPagesEditor/Private/TabFactory/PageViewerTabSummoner.h
new file mode 100644
--- /dev/null
+++ b/Engine/Plugins/Experimental/RenderPages/Source/RenderPagesEditor/Private/TabFactory/PageViewerTabSummoner.h
@@ -0,0 +1,34 @@
+// Copyright Epic Games, Inc. All Rights Reserved.
+
+#pragma once
+
+#include "Widgets/SWidget.h"
+#include "WorkflowOrientedApp/WorkflowTabFactory.h"
+
+
+namespace UE::RenderPages
+{
+	class IRenderPageCollectionEditor;
+}
+
+
+namespace UE::RenderPages::Private
+{
+	/**
+	 * The page viewer tab factory.
+	 */
+	struct FPageViewerTabSummoner : FWorkflowTabFactory
+	{
+	public:
+		/** Unique ID representing this tab. */
+		static const FName TabID;
+
+	public:
+		FPageViewerTabSummoner(TSharedPtr<IRenderPageCollectionEditor> InBlueprintEditor);
+		virtual TSharedRef<SWidget> CreateTabBody(const FWorkflowTabSpawnInfo& Info) const override;
+
+	protected:
+		/** A weak reference to the blueprint editor. */
+		TWeakPtr<IRenderPageCollectionEditor> BlueprintEditorWeakPtr;
+	};
+}
diff --git a/Engine/Plugins/Experimental/RenderPages/Source/RenderPagesEditor/Private/Toolkit/RenderPageCollectionEditor.cpp b/Engine/Plugins/Experimental/RenderPages/Source/RenderPagesEditor/Private/Toolkit/RenderPageCollectionEditor.cpp
new file mode 100644
--- /dev/null
+++ b/Engine/Plugins/Experimental/RenderPages/Source/RenderPagesEditor/Private/Toolkit/RenderPageCollectionEditor.cpp
@@ -0,0 +1,770 @@
+﻿// Copyright Epic Games, Inc. All Rights Reserved.
+
+#include "Toolkit/RenderPageCollectionEditor.h"
+
+#include "RenderPage/RenderPageCollection.h"
+#include "RenderPage/RenderPageManager.h"
+#include "RenderPage/RenderPageMoviePipelineJob.h"
+#include "RenderPageCollectionEditorToolbar.h"
+#include "IRenderPagesEditorModule.h"
+#include "IRenderPagesModule.h"
+
+#include "BlueprintCompilationManager.h"
+#include "BlueprintModes/RenderPagesApplicationModeBase.h"
+#include "BlueprintModes/RenderPagesApplicationModeListing.h"
+#include "BlueprintModes/RenderPagesApplicationModeLogic.h"
+#include "BlueprintModes/RenderPagesApplicationModes.h"
+#include "Blueprints/RenderPagesBlueprint.h"
+#include "Commands/RenderPagesEditorCommands.h"
+#include "EdGraphNode_Comment.h"
+#include "EditorModeManager.h"
+#include "Framework/Commands/GenericCommands.h"
+#include "Framework/MultiBox/MultiBoxBuilder.h"
+#include "Graph/RenderPagesGraph.h"
+#include "Kismet2/BlueprintEditorUtils.h"
+#include "Misc/MessageDialog.h"
+#include "PropertyCustomizationHelpers.h"
+#include "SBlueprintEditorToolbar.h"
+#include "SGraphPanel.h"
+#include "SKismetInspector.h"
+#include "SMyBlueprint.h"
+#include "SNodePanel.h"
+#include "Stats/StatsHierarchical.h"
+#include "Styling/AppStyle.h"
+
+#define LOCTEXT_NAMESPACE "FRenderPageCollectionEditor"
+
+
+namespace UE::RenderPages::Private
+{
+	const FName RenderPagesEditorAppName(TEXT("RenderPagesEditorApp"));
+
+	struct FRenderPagesZoomLevelsContainer : FZoomLevelsContainer
+	{
+		struct FRenderPagesZoomLevelEntry
+		{
+		public:
+			FRenderPagesZoomLevelEntry(const float InZoomAmount, const FText& InDisplayText, const EGraphRenderingLOD::Type InLOD)
+				: DisplayText(FText::Format(NSLOCTEXT("GraphEditor", "Zoom", "Zoom {0}"), InDisplayText))
+				, ZoomAmount(InZoomAmount)
+				, LOD(InLOD)
+			{}
+
+		public:
+			FText DisplayText;
+			float ZoomAmount;
+			EGraphRenderingLOD::Type LOD;
+		};
+
+		FRenderPagesZoomLevelsContainer()
+		{
+			ZoomLevels.Reserve(22);
+			ZoomLevels.Add(FRenderPagesZoomLevelEntry(0.025f, FText::FromString(TEXT("-14")), EGraphRenderingLOD::LowestDetail));
+			ZoomLevels.Add(FRenderPagesZoomLevelEntry(0.070f, FText::FromString(TEXT("-13")), EGraphRenderingLOD::LowestDetail));
+			ZoomLevels.Add(FRenderPagesZoomLevelEntry(0.100f, FText::FromString(TEXT("-12")), EGraphRenderingLOD::LowestDetail));
+			ZoomLevels.Add(FRenderPagesZoomLevelEntry(0.125f, FText::FromString(TEXT("-11")), EGraphRenderingLOD::LowestDetail));
+			ZoomLevels.Add(FRenderPagesZoomLevelEntry(0.150f, FText::FromString(TEXT("-10")), EGraphRenderingLOD::LowestDetail));
+			ZoomLevels.Add(FRenderPagesZoomLevelEntry(0.175f, FText::FromString(TEXT("-9")), EGraphRenderingLOD::LowestDetail));
+			ZoomLevels.Add(FRenderPagesZoomLevelEntry(0.200f, FText::FromString(TEXT("-8")), EGraphRenderingLOD::LowestDetail));
+			ZoomLevels.Add(FRenderPagesZoomLevelEntry(0.225f, FText::FromString(TEXT("-7")), EGraphRenderingLOD::LowDetail));
+			ZoomLevels.Add(FRenderPagesZoomLevelEntry(0.250f, FText::FromString(TEXT("-6")), EGraphRenderingLOD::LowDetail));
+			ZoomLevels.Add(FRenderPagesZoomLevelEntry(0.375f, FText::FromString(TEXT("-5")), EGraphRenderingLOD::MediumDetail));
+			ZoomLevels.Add(FRenderPagesZoomLevelEntry(0.500f, FText::FromString(TEXT("-4")), EGraphRenderingLOD::MediumDetail));
+			ZoomLevels.Add(FRenderPagesZoomLevelEntry(0.675f, FText::FromString(TEXT("-3")), EGraphRenderingLOD::MediumDetail));
+			ZoomLevels.Add(FRenderPagesZoomLevelEntry(0.750f, FText::FromString(TEXT("-2")), EGraphRenderingLOD::DefaultDetail));
+			ZoomLevels.Add(FRenderPagesZoomLevelEntry(0.875f, FText::FromString(TEXT("-1")), EGraphRenderingLOD::DefaultDetail));
+			ZoomLevels.Add(FRenderPagesZoomLevelEntry(1.000f, FText::FromString(TEXT("1:1")), EGraphRenderingLOD::DefaultDetail));
+			ZoomLevels.Add(FRenderPagesZoomLevelEntry(1.250f, FText::FromString(TEXT("+1")), EGraphRenderingLOD::DefaultDetail));
+			ZoomLevels.Add(FRenderPagesZoomLevelEntry(1.375f, FText::FromString(TEXT("+2")), EGraphRenderingLOD::DefaultDetail));
+			ZoomLevels.Add(FRenderPagesZoomLevelEntry(1.500f, FText::FromString(TEXT("+3")), EGraphRenderingLOD::FullyZoomedIn));
+			ZoomLevels.Add(FRenderPagesZoomLevelEntry(1.675f, FText::FromString(TEXT("+4")), EGraphRenderingLOD::FullyZoomedIn));
+			ZoomLevels.Add(FRenderPagesZoomLevelEntry(1.750f, FText::FromString(TEXT("+5")), EGraphRenderingLOD::FullyZoomedIn));
+			ZoomLevels.Add(FRenderPagesZoomLevelEntry(1.875f, FText::FromString(TEXT("+6")), EGraphRenderingLOD::FullyZoomedIn));
+			ZoomLevels.Add(FRenderPagesZoomLevelEntry(2.000f, FText::FromString(TEXT("+7")), EGraphRenderingLOD::FullyZoomedIn));
+		}
+
+		virtual float GetZoomAmount(int32 InZoomLevel) const override
+		{
+			checkSlow(ZoomLevels.IsValidIndex(InZoomLevel));
+			return ZoomLevels[InZoomLevel].ZoomAmount;
+		}
+
+		virtual int32 GetNearestZoomLevel(float InZoomAmount) const override
+		{
+			for (int32 ZoomLevelIndex = 0; ZoomLevelIndex < GetNumZoomLevels(); ++ZoomLevelIndex)
+			{
+				if (InZoomAmount <= GetZoomAmount(ZoomLevelIndex))
+				{
+					return ZoomLevelIndex;
+				}
+			}
+
+			return GetDefaultZoomLevel();
+		}
+
+		virtual FText GetZoomText(int32 InZoomLevel) const override
+		{
+			checkSlow(ZoomLevels.IsValidIndex(InZoomLevel));
+			return ZoomLevels[InZoomLevel].DisplayText;
+		}
+
+		virtual int32 GetNumZoomLevels() const override
+		{
+			return ZoomLevels.Num();
+		}
+
+		virtual int32 GetDefaultZoomLevel() const override
+		{
+			return 14;
+		}
+
+		virtual EGraphRenderingLOD::Type GetLOD(int32 InZoomLevel) const override
+		{
+			checkSlow(ZoomLevels.IsValidIndex(InZoomLevel));
+			return ZoomLevels[InZoomLevel].LOD;
+		}
+
+		TArray<FRenderPagesZoomLevelEntry> ZoomLevels;
+	};
+}
+
+void UE::RenderPages::Private::FRenderPageCollectionEditor::InitRenderPagesEditor(const EToolkitMode::Type Mode, const TSharedPtr<IToolkitHost>& InitToolkitHost, URenderPagesBlueprint* InRenderPagesBlueprint)
+{
+	check(InRenderPagesBlueprint);
+
+	FBlueprintCompilationManager::FlushCompilationQueue(nullptr);
+
+	TSharedPtr<IRenderPageCollectionEditor> ThisPtr(SharedThis(this));
+
+	if (!Toolbar.IsValid())
+	{
+		Toolbar = MakeShareable(new FBlueprintEditorToolbar(SharedThis(this)));
+	}
+
+	RenderPagesToolbar = MakeShared<FRenderPagesBlueprintEditorToolbar>(ThisPtr);
+
+	// Build up a list of objects being edited in this asset editor
+	TArray<UObject*> ObjectsBeingEdited;
+	ObjectsBeingEdited.Add(InRenderPagesBlueprint);
+
+	// Initialize the asset editor and spawn tabs
+	constexpr bool bCreateDefaultStandaloneMenu = true;
+	constexpr bool bCreateDefaultToolbar = true;
+	InitAssetEditor(Mode, InitToolkitHost, RenderPagesEditorAppName, FTabManager::FLayout::NullLayout, bCreateDefaultStandaloneMenu, bCreateDefaultToolbar, ObjectsBeingEdited);
+
+	CreateDefaultCommands();
+
+	TArray<UBlueprint*> RenderPagesBlueprints;
+	RenderPagesBlueprints.Add(InRenderPagesBlueprint);
+
+	CommonInitialization(RenderPagesBlueprints, false);
+
+	{
+		TArray<UEdGraph*> EdGraphs;
+		InRenderPagesBlueprint->GetAllGraphs(EdGraphs);
+
+		for (UEdGraph* Graph : EdGraphs)
+		{
+			if (URenderPagesGraph* RenderPagesGraph = Cast<URenderPagesGraph>(Graph))
+			{
+				RenderPagesGraph->Initialize(InRenderPagesBlueprint);
+			}
+		}
+	}
+
+	BindCommands();
+
+	ExtendMenu();
+	ExtendToolbar();
+	RegenerateMenusAndToolbars();
+
+	UpdateInstance(GetRenderPagesBlueprint(), true);
+
+	constexpr bool bShouldOpenInDefaultsMode = true;
+	RegisterApplicationModes(RenderPagesBlueprints, bShouldOpenInDefaultsMode, InRenderPagesBlueprint->bIsNewlyCreated);
+
+	// Post-layout initialization
+	PostLayoutBlueprintEditorInitialization();
+}
+
+UE::RenderPages::Private::FRenderPageCollectionEditor::FRenderPageCollectionEditor()
+	: PreviewBlueprint(nullptr)
+	, bRunRenderNewBatch(false)
+	, BatchRenderJob(nullptr)
+	, PreviewRenderJob(nullptr)
+{}
+
+UE::RenderPages::Private::FRenderPageCollectionEditor::~FRenderPageCollectionEditor()
+{
+	URenderPagesBlueprint* RenderPagesBlueprint = FRenderPageCollectionEditor::GetRenderPagesBlueprint();
+
+	RenderPagesEditorClosedDelegate.Broadcast(this, RenderPagesBlueprint);
+
+	if (RenderPagesBlueprint)
+	{
+		// clear editor related data from the debugged control rig instance 
+		RenderPagesBlueprint->SetObjectBeingDebugged(nullptr);
+		RenderPagesBlueprint = nullptr;
+	}
+	DestroyInstance();
+}
+
+void UE::RenderPages::Private::FRenderPageCollectionEditor::CreateDefaultCommands()
+{
+	if (GetBlueprintObj())
+	{
+		FBlueprintEditor::CreateDefaultCommands();
+	}
+	else
+	{
+		ToolkitCommands->MapAction(FGenericCommands::Get().Undo, FExecuteAction::CreateSP(this, &FRenderPageCollectionEditor::UndoAction));
+		ToolkitCommands->MapAction(FGenericCommands::Get().Redo, FExecuteAction::CreateSP(this, &FRenderPageCollectionEditor::RedoAction));
+	}
+}
+
+TSharedRef<SGraphEditor> UE::RenderPages::Private::FRenderPageCollectionEditor::CreateGraphEditorWidget(TSharedRef<FTabInfo> InTabInfo, UEdGraph* InGraph)
+{
+	TSharedRef<SGraphEditor> GraphEditor = FBlueprintEditor::CreateGraphEditorWidget(InTabInfo, InGraph);
+
+	GraphEditor->GetGraphPanel()->SetZoomLevelsContainer<FRenderPagesZoomLevelsContainer>();
+
+	return GraphEditor;
+}
+
+UBlueprint* UE::RenderPages::Private::FRenderPageCollectionEditor::GetBlueprintObj() const
+{
+	for (UObject* Obj : GetEditingObjects())
+	{
+		if (!IsValid(Obj))
+		{
+			continue;
+		}
+		if (UBlueprint* Result = Cast<UBlueprint>(Obj))
+		{
+			return Result;
+		}
+	}
+	return nullptr;
+}
+
+FGraphAppearanceInfo UE::RenderPages::Private::FRenderPageCollectionEditor::GetGraphAppearance(UEdGraph* InGraph) const
+{
+	FGraphAppearanceInfo AppearanceInfo = FBlueprintEditor::GetGraphAppearance(InGraph);
+
+	if (GetBlueprintObj()->IsA(URenderPagesBlueprint::StaticClass()))
+	{
+		AppearanceInfo.CornerText = LOCTEXT("AppearanceCornerText_RenderPages", "RENDER PAGES");
+	}
+
+	return AppearanceInfo;
+}
+
+bool UE::RenderPages::Private::FRenderPageCollectionEditor::IsEditable(UEdGraph* InGraph) const
+{
+	return IsGraphInCurrentBlueprint(InGraph);
+}
+
+bool UE::RenderPages::Private::FRenderPageCollectionEditor::IsCompilingEnabled() const
+{
+	return true;
+}
+
+void UE::RenderPages::Private::FRenderPageCollectionEditor::NewDocument_OnClicked(ECreatedDocumentType GraphType)
+{
+	if (GraphType != CGT_NewFunctionGraph)
+	{
+		return;
+	}
+}
+
+bool UE::RenderPages::Private::FRenderPageCollectionEditor::IsSectionVisible(NodeSectionID::Type InSectionID) const
+{
+	switch (InSectionID)
+	{
+		case NodeSectionID::GRAPH:
+		case NodeSectionID::VARIABLE:
+		case NodeSectionID::FUNCTION:
+		{
+			return true;
+		}
+		case NodeSectionID::LOCAL_VARIABLE:
+		{
+			return IsValid(Cast<URenderPagesGraph>(GetFocusedGraph()));
+		}
+		default:
+		{
+			break;
+		}
+	}
+	return false;
+}
+
+FText UE::RenderPages::Private::FRenderPageCollectionEditor::GetGraphDecorationString(UEdGraph* InGraph) const
+{
+	return FText::GetEmpty();
+}
+
+void UE::RenderPages::Private::FRenderPageCollectionEditor::OnActiveTabChanged(TSharedPtr<SDockTab> PreviouslyActive, TSharedPtr<SDockTab> NewlyActivated)
+{
+	if (!NewlyActivated.IsValid())
+	{
+		TArray<UObject*> ObjArray;
+		Inspector->ShowDetailsForObjects(ObjArray);
+	}
+	else
+	{
+		FBlueprintEditor::OnActiveTabChanged(PreviouslyActive, NewlyActivated);
+	}
+}
+
+void UE::RenderPages::Private::FRenderPageCollectionEditor::SetupGraphEditorEvents(UEdGraph* InGraph, SGraphEditor::FGraphEditorEvents& InEvents)
+{
+	FBlueprintEditor::SetupGraphEditorEvents(InGraph, InEvents);
+
+	InEvents.OnCreateActionMenu = SGraphEditor::FOnCreateActionMenu::CreateSP(this, &FRenderPageCollectionEditor::HandleCreateGraphActionMenu);
+}
+
+void UE::RenderPages::Private::FRenderPageCollectionEditor::RegisterApplicationModes(const TArray<UBlueprint*>& InBlueprints, bool bShouldOpenInDefaultsMode, bool bNewlyCreated)
+{
+	if (InBlueprints.Num() == 1)
+	{
+		TSharedPtr<FRenderPageCollectionEditor> ThisPtr(SharedThis(this));
+
+		// Create the modes and activate one (which will populate with a real layout)
+		TArray<TSharedRef<FApplicationMode>> TempModeList;
+		TempModeList.Add(MakeShareable(new FRenderPagesApplicationModeListing(ThisPtr)));
+		TempModeList.Add(MakeShareable(new FRenderPagesApplicationModeLogic(ThisPtr)));
+
+		for (TSharedRef<FApplicationMode>& AppMode : TempModeList)
+		{
+			AddApplicationMode(AppMode->GetModeName(), AppMode);
+		}
+
+		SetCurrentMode(FRenderPagesApplicationModes::ListingMode);
+
+		// Activate our edit mode
+		GetEditorModeManager().SetDefaultMode(FRenderPagesApplicationModes::ListingMode);
+		GetEditorModeManager().ActivateMode(FRenderPagesApplicationModes::ListingMode);
+	}
+}
+
+void UE::RenderPages::Private::FRenderPageCollectionEditor::FocusInspectorOnGraphSelection(const TSet<UObject*>& NewSelection, bool bForceRefresh)
+{
+	// nothing to do here for render pages
+}
+
+bool UE::RenderPages::Private::FRenderPageCollectionEditor::CanAddNewLocalVariable() const
+{
+	return (Cast<URenderPagesGraph>(GetFocusedGraph()) != nullptr);
+}
+
+void UE::RenderPages::Private::FRenderPageCollectionEditor::OnAddNewLocalVariable()
+{
+	if (!CanAddNewLocalVariable())
+	{
+		return;
+	}
+}
+
+void UE::RenderPages::Private::FRenderPageCollectionEditor::OnPasteNewLocalVariable(const FBPVariableDescription& VariableDescription)
+{
+	if (!CanAddNewLocalVariable())
+	{
+		return;
+	}
+}
+
+void UE::RenderPages::Private::FRenderPageCollectionEditor::Compile()
+{
+	DestroyInstance();
+
+	FBlueprintEditor::Compile();
+}
+
+URenderPagesBlueprint* UE::RenderPages::Private::FRenderPageCollectionEditor::GetRenderPagesBlueprint() const
+{
+	return Cast<URenderPagesBlueprint>(GetBlueprintObj());
+}
+
+URenderPageCollection* UE::RenderPages::Private::FRenderPageCollectionEditor::GetInstance() const
+{
+	return RenderPageCollectionWeakPtr.Get();
+}
+
+bool UE::RenderPages::Private::FRenderPageCollectionEditor::IsBatchRendering() const
+{
+	return BatchRenderJob && IsValid(BatchRenderJob);
+}
+
+bool UE::RenderPages::Private::FRenderPageCollectionEditor::IsPreviewRendering() const
+{
+	return PreviewRenderJob && IsValid(PreviewRenderJob);
+}
+
+void UE::RenderPages::Private::FRenderPageCollectionEditor::MarkAsModified()
+{
+	if (URenderPageCollection* Instance = GetInstance(); IsValid(Instance))
+	{
+		Instance->Modify();
+	}
+	if (UBlueprint* BlueprintObj = GetBlueprintObj(); IsValid(BlueprintObj))
+	{
+		BlueprintObj->Modify();
+	}
+}
+
+TArray<URenderPage*> UE::RenderPages::Private::FRenderPageCollectionEditor::GetSelectedRenderPages() const
+{
+	TArray<URenderPage*> SelectedPages;
+	if (URenderPageCollection* Collection = GetInstance(); IsValid(Collection))
+	{
+		for (URenderPage* Page : Collection->GetRenderPages())
+		{
+			if (SelectedRenderPagesIds.Contains(Page->GetId()))
+			{
+				SelectedPages.Add(Page);
+			}
+		}
+	}
+	return SelectedPages;
+}
+
+void UE::RenderPages::Private::FRenderPageCollectionEditor::SetSelectedRenderPages(const TArray<URenderPage*>& RenderPages)
+{
+	SelectedRenderPagesIds.Empty();
+	for (URenderPage* Page : RenderPages)
+	{
+		if (!IsValid(Page))
+		{
+			continue;
+		}
+		SelectedRenderPagesIds.Add(Page->GetId());
+	}
+
+	OnRenderPagesSelectionChanged().Broadcast();
+}
+
+FName UE::RenderPages::Private::FRenderPageCollectionEditor::GetToolkitFName() const
+{
+	return FName("RenderPageCollectionEditor");
+}
+
+FText UE::RenderPages::Private::FRenderPageCollectionEditor::GetBaseToolkitName() const
+{
+	return LOCTEXT("AppLabel", "Render Page Collection Editor");
+}
+
+FString UE::RenderPages::Private::FRenderPageCollectionEditor::GetDocumentationLink() const
+{
+	return FString();
+}
+
+FText UE::RenderPages::Private::FRenderPageCollectionEditor::GetToolkitToolTipText() const
+{
+	return GetToolTipTextForObject(GetBlueprintObj());
+}
+
+FLinearColor UE::RenderPages::Private::FRenderPageCollectionEditor::GetWorldCentricTabColorScale() const
+{
+	return FLinearColor(0.5f, 0.25f, 0.35f, 0.5f);
+}
+
+FString UE::RenderPages::Private::FRenderPageCollectionEditor::GetWorldCentricTabPrefix() const
+{
+	return LOCTEXT("WorldCentricTabPrefix", "Render Page Collection ").ToString();
+}
+
+void UE::RenderPages::Private::FRenderPageCollectionEditor::InitToolMenuContext(FToolMenuContext& MenuContext)
+{
+	FBlueprintEditor::InitToolMenuContext(MenuContext);
+}
+
+void UE::RenderPages::Private::FRenderPageCollectionEditor::OnToolkitHostingStarted(const TSharedRef<IToolkit>& Toolkit) {}
+
+void UE::RenderPages::Private::FRenderPageCollectionEditor::OnToolkitHostingFinished(const TSharedRef<IToolkit>& Toolkit) {}
+
+void UE::RenderPages::Private::FRenderPageCollectionEditor::Tick(float DeltaTime)
+{
+	FBlueprintEditor::Tick(DeltaTime);
+
+	// Note: The weak ptr can become stale if the actor is reinstanced due to a Blueprint change, etc. In that case we 
+	//       look to see if we can find the new instance in the preview world and then update the weak ptr.
+	if (RenderPageCollectionWeakPtr.IsStale(true))
+	{
+		RefreshInstance();
+	}
+
+	if (bRunRenderNewBatch)
+	{
+		bRunRenderNewBatch = false;
+		BatchRenderListAction();
+	}
+}
+
+TStatId UE::RenderPages::Private::FRenderPageCollectionEditor::GetStatId() const
+{
+	RETURN_QUICK_DECLARE_CYCLE_STAT(FRenderPageCollectionEditor, STATGROUP_Tickables);
+}
+
+void UE::RenderPages::Private::FRenderPageCollectionEditor::RefreshInstance()
+{
+	UpdateInstance(GetRenderPagesBlueprint(), true);
+}
+
+void UE::RenderPages::Private::FRenderPageCollectionEditor::OnBlueprintChangedImpl(UBlueprint* InBlueprint, bool bIsJustBeingCompiled)
+{
+	DestroyInstance();
+
+	FBlueprintEditor::OnBlueprintChangedImpl(InBlueprint, bIsJustBeingCompiled);
+
+	if (InBlueprint)
+	{
+		RefreshInstance();
+	}
+}
+
+void UE::RenderPages::Private::FRenderPageCollectionEditor::AddReferencedObjects(FReferenceCollector& Collector)
+{
+	FBlueprintEditor::AddReferencedObjects(Collector);
+
+	URenderPageCollection* RenderPageCollection = GetInstance();
+	Collector.AddReferencedObject(RenderPageCollection);
+}
+
+void UE::RenderPages::Private::FRenderPageCollectionEditor::NotifyPreChange(FProperty* PropertyAboutToChange)
+{
+	FBlueprintEditor::NotifyPreChange(PropertyAboutToChange);
+
+	if (URenderPagesBlueprint* RenderPagesBP = GetRenderPagesBlueprint())
+	{
+		RenderPagesBP->Modify();
+	}
+}
+
+void UE::RenderPages::Private::FRenderPageCollectionEditor::NotifyPostChange(const FPropertyChangedEvent& PropertyChangedEvent, FProperty* PropertyThatChanged)
+{
+	FBlueprintEditor::NotifyPostChange(PropertyChangedEvent, PropertyThatChanged);
+}
+
+void UE::RenderPages::Private::FRenderPageCollectionEditor::OnFinishedChangingProperties(const FPropertyChangedEvent& PropertyChangedEvent) {}
+
+void UE::RenderPages::Private::FRenderPageCollectionEditor::BindCommands()
+{
+	const auto& Commands = FRenderPagesEditorCommands::Get();
+
+	ToolkitCommands->MapAction(Commands.AddPage, FExecuteAction::CreateSP(this, &FRenderPageCollectionEditor::AddPageAction));
+	ToolkitCommands->MapAction(Commands.CopyPage, FExecuteAction::CreateSP(this, &FRenderPageCollectionEditor::CopyPageAction));
+	ToolkitCommands->MapAction(Commands.DeletePage, FExecuteAction::CreateSP(this, &FRenderPageCollectionEditor::DeletePageAction));
+	ToolkitCommands->MapAction(Commands.BatchRenderList, FExecuteAction::CreateSP(this, &FRenderPageCollectionEditor::BatchRenderListAction));
+	ToolkitCommands->MapAction(Commands.AutoCompileGraph, FExecuteAction::CreateSP(this, &FRenderPageCollectionEditor::AutoCompileGraphAction));
+}
+
+void UE::RenderPages::Private::FRenderPageCollectionEditor::AddPageAction()
+{
+	URenderPage* Page = IRenderPagesModule::Get().GetManager().AddNewPage(GetInstance());
+	if (!IsValid(Page))
+	{
+		return;
+	}
+
+	OnRenderPageCreated().Broadcast(Page);
+	MarkAsModified();
+	OnRenderPagesChanged().Broadcast();
+}
+
+void UE::RenderPages::Private::FRenderPageCollectionEditor::CopyPageAction()
+{
+	TArray<URenderPage*> SelectedRenderPages = GetSelectedRenderPages();
+	if (SelectedRenderPages.Num() <= 0)
+	{
+		return;
+	}
+
+	URenderPageCollection* RenderPageCollection = GetInstance();
+	for (URenderPage* SelectedRenderPage : SelectedRenderPages)
+	{
+		IRenderPagesModule::Get().GetManager().CopyPage(RenderPageCollection, SelectedRenderPage);
+	}
+
+	MarkAsModified();
+	OnRenderPagesChanged().Broadcast();
+}
+
+void UE::RenderPages::Private::FRenderPageCollectionEditor::DeletePageAction()
+{
+	TArray<URenderPage*> SelectedRenderPages = GetSelectedRenderPages();
+	if (SelectedRenderPages.Num() <= 0)
+	{
+		return;
+	}
+
+	const FText TitleText = LOCTEXT("ConfirmToDeleteTitle", "Confirm To Delete");
+	const EAppReturnType::Type DialogResult = FMessageDialog::Open(
+		EAppMsgType::OkCancel,
+		((SelectedRenderPages.Num() == 1) ? LOCTEXT("ConfirmToDeleteSingleText", "Are you sure you want to delete the selected render page?") : LOCTEXT("ConfirmToDeleteMultipleText", "Are you sure you want to delete the selected render pages?")),
+		&TitleText);
+
+	if (DialogResult != EAppReturnType::Ok)
+	{
+		return;
+	}
+
+	URenderPageCollection* RenderPageCollection = GetInstance();
+	for (URenderPage* SelectedRenderPage : SelectedRenderPages)
+	{
+		IRenderPagesModule::Get().GetManager().DeletePage(RenderPageCollection, SelectedRenderPage);
+	}
+
+	MarkAsModified();
+	OnRenderPagesChanged().Broadcast();
+}
+
+void UE::RenderPages::Private::FRenderPageCollectionEditor::BatchRenderListAction()
+{
+	if (!CanCurrentlyRender())
+	{
+		bRunRenderNewBatch = true;
+		return;
+	}
+
+	URenderPageCollection* PageCollection = GetInstance();
+	if (PageCollection->GetRenderPages().Num() <= 0)
+	{
+		const FText TitleText = LOCTEXT("NoPagesToRenderTitle", "No Pages To Render");
+		FMessageDialog::Open(
+			EAppMsgType::Ok,
+			LOCTEXT("NoPagesToRenderText", "There are no render pages in this collection, and so nothing can be rendered. Please make a page and try again."),
+			&TitleText);
+		return;
+	}
+
+	URenderPagesMoviePipelineRenderJob* RenderJob = IRenderPagesModule::Get().GetManager().CreateBatchRenderJob(PageCollection);
+	if (!RenderJob)
+	{
+		return;
+	}
+	RenderJob->OnExecuteFinished().AddRaw(this, &FRenderPageCollectionEditor::OnBatchRenderListActionFinished);
+	BatchRenderJob = RenderJob;
+	OnRenderPagesBatchRenderingStarted().Broadcast(RenderJob);
+	RenderJob->Execute();
+}
+
+void UE::RenderPages::Private::FRenderPageCollectionEditor::OnBatchRenderListActionFinished(URenderPagesMoviePipelineRenderJob* RenderJob, bool bSuccess)
+{
+	URenderPagesMoviePipelineRenderJob* FinishedRenderJob = BatchRenderJob;
+	BatchRenderJob = nullptr;
+	OnRenderPagesBatchRenderingFinished().Broadcast(FinishedRenderJob);
+}
+
+void UE::RenderPages::Private::FRenderPageCollectionEditor::AutoCompileGraphAction() {}
+
+void UE::RenderPages::Private::FRenderPageCollectionEditor::UndoAction()
+{
+	GEditor->UndoTransaction();
+}
+
+void UE::RenderPages::Private::FRenderPageCollectionEditor::RedoAction()
+{
+	GEditor->RedoTransaction();
+}
+
+void UE::RenderPages::Private::FRenderPageCollectionEditor::ExtendMenu()
+{
+	if (MenuExtender.IsValid())
+	{
+		RemoveMenuExtender(MenuExtender);
+		MenuExtender.Reset();
+	}
+
+	MenuExtender = MakeShareable(new FExtender);
+
+	AddMenuExtender(MenuExtender);
+
+	// add extensible menu if exists
+	IRenderPagesEditorModule& RenderPagesEditorModule = IRenderPagesEditorModule::Get();
+	AddMenuExtender(RenderPagesEditorModule.GetMenuExtensibilityManager()->GetAllExtenders(GetToolkitCommands(), GetEditingObjects()));
+}
+
+void UE::RenderPages::Private::FRenderPageCollectionEditor::ExtendToolbar()
+{
+	// If the ToolbarExtender is valid, remove it before rebuilding it
+	if (ToolbarExtender.IsValid())
+	{
+		RemoveToolbarExtender(ToolbarExtender);
+		ToolbarExtender.Reset();
+	}
+
+	ToolbarExtender = MakeShareable(new FExtender);
+
+	AddToolbarExtender(ToolbarExtender);
+
+	IRenderPagesEditorModule& RenderPagesEditorModule = IRenderPagesEditorModule::Get();
+	AddToolbarExtender(RenderPagesEditorModule.GetToolBarExtensibilityManager()->GetAllExtenders(GetToolkitCommands(), GetEditingObjects()));
+
+	ToolbarExtender->AddToolBarExtension(
+		"Asset",
+		EExtensionHook::After,
+		GetToolkitCommands(),
+		FToolBarExtensionDelegate::CreateSP(this, &FRenderPageCollectionEditor::FillToolbar)
+	);
+}
+
+void UE::RenderPages::Private::FRenderPageCollectionEditor::FillToolbar(FToolBarBuilder& ToolbarBuilder)
+{
+	ToolbarBuilder.BeginSection("Common");
+	{
+		ToolbarBuilder.AddToolBarButton(FRenderPagesEditorCommands::Get().BatchRenderList, NAME_None, TAttribute<FText>(), TAttribute<FText>(), FSlateIcon(FAppStyle::GetAppStyleSetName(), "LevelEditor.OpenCinematic"));
+	}
+	ToolbarBuilder.EndSection();
+}
+
+void UE::RenderPages::Private::FRenderPageCollectionEditor::DestroyInstance()
+{
+	if (URenderPageCollection* RenderPageCollection = GetInstance())
+	{
+		RenderPageCollection->OnClose();
+		RenderPageCollectionWeakPtr.Reset();
+		RenderPageCollection->MarkAsGarbage();
+	}
+}
+
+void UE::RenderPages::Private::FRenderPageCollectionEditor::UpdateInstance(UBlueprint* InBlueprint, bool bInForceFullUpdate)
+{
+	// If the Blueprint is changing
+	if ((InBlueprint != PreviewBlueprint.Get()) || bInForceFullUpdate)
+	{
+		// Destroy the previous instance
+		DestroyInstance();
+
+		// Save the Blueprint we're creating a preview for
+		PreviewBlueprint = Cast<URenderPagesBlueprint>(InBlueprint);
+
+		URenderPageCollection* RenderPageCollection;
+		// Create the Widget, we have to do special swapping out of the widget tree.
+		{
+			// Assign the outer to the game instance if it exists, otherwise use the world
+			{
+				FMakeClassSpawnableOnScope TemporarilySpawnable(PreviewBlueprint->GeneratedClass);
+				RenderPageCollection = NewObject<URenderPageCollection>(PreviewScene.GetWorld(), PreviewBlueprint->GeneratedClass);
+			}
+		}
+
+		// Make sure the object being debugged is the preview instance
+		GetBlueprintObj()->SetObjectBeingDebugged(RenderPageCollection);
+
+		// Store a reference to the preview actor.
+		RenderPageCollectionWeakPtr = RenderPageCollection;
+
+		// Broadcast the events
+		OnRenderPagesChanged().Broadcast();
+		OnRenderPagesSelectionChanged().Broadcast();
+	}
+}
+
+FActionMenuContent UE::RenderPages::Private::FRenderPageCollectionEditor::HandleCreateGraphActionMenu(UEdGraph* InGraph, const FVector2D& InNodePosition, const TArray<UEdGraphPin*>& InDraggedPins, bool bAutoExpand, SGraphEditor::FActionMenuClosed InOnMenuClosed)
+{
+	return OnCreateGraphActionMenu(InGraph, InNodePosition, InDraggedPins, bAutoExpand, InOnMenuClosed);
+}
+
+
+#undef LOCTEXT_NAMESPACE
diff --git a/Engine/Plugins/Experimental/RenderPages/Source/RenderPagesEditor/Private/Toolkit/RenderPageCollectionEditor.h b/Engine/Plugins/Experimental/RenderPages/Source/RenderPagesEditor/Private/Toolkit/RenderPageCollectionEditor.h
new file mode 100644
--- /dev/null
+++ b/Engine/Plugins/Experimental/RenderPages/Source/RenderPagesEditor/Private/Toolkit/RenderPageCollectionEditor.h
@@ -0,0 +1,199 @@
+﻿// Copyright Epic Games, Inc. All Rights Reserved.
+
+#pragma once
+
+#include "CoreMinimal.h"
+#include "IRenderPageCollectionEditor.h"
+
+
+class FSpawnTabArgs;
+class IToolkitHost;
+class SDockTab;
+class URenderPageCollection;
+class URenderPage;
+class URenderPagesBlueprint;
+class URenderPagesMoviePipelineRenderJob;
+
+namespace UE::RenderPages::Private
+{
+	class FRenderPagesBlueprintEditorToolbar;
+}
+
+
+namespace UE::RenderPages::Private
+{
+	DECLARE_MULTICAST_DELEGATE_TwoParams(FRenderPagesEditorClosed, const UE::RenderPages::IRenderPageCollectionEditor*, URenderPagesBlueprint*);
+
+
+	/**
+	 * The render page editor implementation.
+	 */
+	class FRenderPageCollectionEditor : public IRenderPageCollectionEditor
+	{
+	public:
+		void InitRenderPagesEditor(const EToolkitMode::Type Mode, const TSharedPtr<IToolkitHost>& InitToolkitHost, URenderPagesBlueprint* InRenderPagesBlueprint);
+
+	public:
+		FRenderPageCollectionEditor();
+		virtual ~FRenderPageCollectionEditor() override;
+
+		//~ Begin FBlueprintEditor Interface
+		virtual void CreateDefaultCommands() override;
+		virtual TSharedRef<SGraphEditor> CreateGraphEditorWidget(TSharedRef<FTabInfo> InTabInfo, UEdGraph* InGraph) override;
+		virtual UBlueprint* GetBlueprintObj() const override;
+		virtual TSharedPtr<FRenderPagesBlueprintEditorToolbar> GetRenderPagesToolbarBuilder() override { return RenderPagesToolbar; };
+		virtual FGraphAppearanceInfo GetGraphAppearance(UEdGraph* InGraph) const override;
+		virtual bool IsEditable(UEdGraph* InGraph) const override;
+		virtual bool IsCompilingEnabled() const override;
+		virtual bool IsInAScriptingMode() const override { return true; }
+		virtual void NewDocument_OnClicked(ECreatedDocumentType GraphType) override;
+		virtual bool IsSectionVisible(NodeSectionID::Type InSectionID) const override;
+		virtual FText GetGraphDecorationString(UEdGraph* InGraph) const override;
+		virtual void OnActiveTabChanged(TSharedPtr<SDockTab> PreviouslyActive, TSharedPtr<SDockTab> NewlyActivated) override;
+		virtual void SetupGraphEditorEvents(UEdGraph* InGraph, SGraphEditor::FGraphEditorEvents& InEvents) override;
+		virtual void RegisterApplicationModes(const TArray<UBlueprint*>& InBlueprints, bool bShouldOpenInDefaultsMode, bool bNewlyCreated = false) override;
+		virtual void FocusInspectorOnGraphSelection(const TSet<UObject*>& NewSelection, bool bForceRefresh = false) override;
+
+		virtual bool ShouldLoadBPLibrariesFromAssetRegistry() override { return false; }// Render Pages BP does not use regular BP function libraries, no need to load them.
+
+		virtual bool CanAddNewLocalVariable() const override;
+		virtual void OnAddNewLocalVariable() override;
+		virtual void OnPasteNewLocalVariable(const FBPVariableDescription& VariableDescription) override;
+
+		virtual void Compile() override;
+		//~ End FBlueprintEditor Interface
+
+		//~ Begin IRenderPageCollectionEditor Interface
+		virtual URenderPagesBlueprint* GetRenderPagesBlueprint() const override;
+		virtual URenderPageCollection* GetInstance() const override;
+		virtual bool IsBatchRendering() const override;
+		virtual URenderPagesMoviePipelineRenderJob* GetBatchRenderJob() const override { return BatchRenderJob; }
+		virtual bool IsPreviewRendering() const override;
+		virtual URenderPagesMoviePipelineRenderJob* GetPreviewRenderJob() const override { return PreviewRenderJob; }
+		virtual void SetPreviewRenderJob(URenderPagesMoviePipelineRenderJob* Job) override { PreviewRenderJob = Job; }
+		virtual void MarkAsModified() override;
+		virtual TArray<URenderPage*> GetSelectedRenderPages() const override;
+		virtual void SetSelectedRenderPages(const TArray<URenderPage*>& RenderPages) override;
+		//~ End IRenderPageCollectionEditor Interface
+
+		//~ Begin IToolkit Interface
+		virtual FName GetToolkitFName() const override;
+		virtual FText GetBaseToolkitName() const override;
+		virtual FString GetDocumentationLink() const override;
+		virtual FText GetToolkitToolTipText() const override;
+		virtual FLinearColor GetWorldCentricTabColorScale() const override;
+		virtual FString GetWorldCentricTabPrefix() const override;
+		virtual void InitToolMenuContext(FToolMenuContext& MenuContext) override;
+		//~ End IToolkit Interface
+
+		//~ Begin IToolkitHost Interface
+		virtual void OnToolkitHostingStarted(const TSharedRef<IToolkit>& Toolkit) override;
+		virtual void OnToolkitHostingFinished(const TSharedRef<IToolkit>& Toolkit) override;
+		//~ End IToolkitHost Interface
+
+		//~ Begin FTickableEditorObject Interface
+		virtual void Tick(float DeltaTime) override;
+		virtual TStatId GetStatId() const override;
+		//~ End FTickableEditorObject Interface
+
+		/** Immediately rebuilds the render page collection that is being shown in the editor. */
+		void RefreshInstance();
+
+		/** The delegate that will fire when this editor closes. */
+		FRenderPagesEditorClosed& OnRenderPageCollectionEditorClosed() { return RenderPagesEditorClosedDelegate; }
+
+	private:
+		/** Called whenever the blueprint is structurally changed. */
+		virtual void OnBlueprintChangedImpl(UBlueprint* InBlueprint, bool bIsJustBeingCompiled = false) override;
+
+	protected:
+		//~ Begin FGCObject Interface
+		virtual void AddReferencedObjects(FReferenceCollector& Collector) override;
+		//~ End FGCObject Interface
+
+		//~ Begin FNotifyHook Interface
+		virtual void NotifyPreChange(FProperty* PropertyAboutToChange) override;
+		virtual void NotifyPostChange(const FPropertyChangedEvent& PropertyChangedEvent, FProperty* PropertyThatChanged) override;
+		virtual void OnFinishedChangingProperties(const FPropertyChangedEvent& PropertyChangedEvent) override;
+		//~ End FNotifyHook Interface
+
+	protected:
+		/** Binds the FRenderPagesEditorCommands commands to functions in this editor. */
+		void BindCommands();
+
+		/** Creates and adds a new render page to the currently viewing render page collection. */
+		void AddPageAction();
+
+		/** Copies the selected render page(s) and adds them to the currently viewing render page collection. */
+		void CopyPageAction();
+
+		/** Removes the currently selected render page(s) from the currently viewing render page collection. */
+		void DeletePageAction();
+
+		/** Renders all the currently enabled render pages. */
+		void BatchRenderListAction();
+
+		/** The callback for when the batch render list action finishes. */
+		void OnBatchRenderListActionFinished(URenderPagesMoviePipelineRenderJob* RenderJob, bool bSuccess);
+
+		/** Toggles auto compiling of the blueprint graph. */
+		void AutoCompileGraphAction();
+
+		/** Undo the last action. */
+		void UndoAction();
+
+		/** Redo the last action that was undone. */
+		void RedoAction();
+
+	private:
+		/** Extends the menu. */
+		void ExtendMenu();
+
+		/** Extends the toolbar. */
+		void ExtendToolbar();
+
+		/** Fills the toolbar with content. */
+		void FillToolbar(FToolBarBuilder& ToolbarBuilder);
+
+		/** Destroy the render page collection instance that is currently visible in the editor. */
+		void DestroyInstance();
+
+		/** Makes a newly compiled/opened render page collection instance visible in the editor. */
+		void UpdateInstance(UBlueprint* InBlueprint, bool bInForceFullUpdate);
+
+		/** Wraps the normal blueprint editor's action menu creation callback. */
+		FActionMenuContent HandleCreateGraphActionMenu(UEdGraph* InGraph, const FVector2D& InNodePosition, const TArray<UEdGraphPin*>& InDraggedPins, bool bAutoExpand, SGraphEditor::FActionMenuClosed InOnMenuClosed);
+
+	private:
+		/** The delegate that will be fired when this editor closes. */
+		FRenderPagesEditorClosed RenderPagesEditorClosedDelegate;
+
+		/** The toolbar builder that is used to customize the toolbar of this editor. */
+		TSharedPtr<FRenderPagesBlueprintEditorToolbar> RenderPagesToolbar;
+
+	protected:
+		/** The extender to pass to the level editor to extend it's window menu. */
+		TSharedPtr<FExtender> MenuExtender;
+
+		/** The toolbar extender of this editor. */
+		TSharedPtr<FExtender> ToolbarExtender;
+
+		/** The blueprint instance that's currently visible in the editor. */
+		TObjectPtr<URenderPagesBlueprint> PreviewBlueprint;
+
+		/** The current render page collection instance that's visible in the editor. */
+		mutable TWeakObjectPtr<URenderPageCollection> RenderPageCollectionWeakPtr;
+
+		/** The IDs of the currently selected render pages. */
+		TSet<FGuid> SelectedRenderPagesIds;
+
+		/** True if it should call BatchRenderListAction() next frame. */
+		bool bRunRenderNewBatch;
+
+		/** The current batch rendering job, if any. */
+		TObjectPtr<URenderPagesMoviePipelineRenderJob> BatchRenderJob;
+
+		/** The current preview rendering job, if any. */
+		TObjectPtr<URenderPagesMoviePipelineRenderJob> PreviewRenderJob;
+	};
+}
diff --git a/Engine/Plugins/Experimental/RenderPages/Source/RenderPagesEditor/Private/Toolkit/RenderPageCollectionEditorToolbar.cpp b/Engine/Plugins/Experimental/RenderPages/Source/RenderPagesEditor/Private/Toolkit/RenderPageCollectionEditorToolbar.cpp
new file mode 100644
--- /dev/null
+++ b/Engine/Plugins/Experimental/RenderPages/Source/RenderPagesEditor/Private/Toolkit/RenderPageCollectionEditorToolbar.cpp
@@ -0,0 +1,159 @@
+﻿// Copyright Epic Games, Inc. All Rights Reserved.
+
+#include "Toolkit/RenderPageCollectionEditorToolbar.h"
+#include "BlueprintModes/RenderPagesApplicationModes.h"
+#include "Commands/RenderPagesEditorCommands.h"
+#include "IRenderPageCollectionEditor.h"
+#include "Styles/RenderPagesEditorStyle.h"
+#include "Framework/MultiBox/MultiBoxBuilder.h"
+#include "IDocumentation.h"
+#include "Kismet2/BlueprintEditorUtils.h"
+#include "Styling/AppStyle.h"
+#include "ToolMenus.h"
+#include "Widgets/Layout/SSpacer.h"
+#include "Widgets/SToolTip.h"
+#include "WorkflowOrientedApp/SModeWidget.h"
+
+#define LOCTEXT_NAMESPACE "RenderPages"
+
+
+//////////////////////////////////////////////////////////////////////////
+// SBlueprintModeSeparator
+
+class SBlueprintModeSeparator : public SBorder
+{
+public:
+	SLATE_BEGIN_ARGS(SBlueprintModeSeparator) {}
+	SLATE_END_ARGS()
+
+	void Construct(const FArguments& InArg)
+	{
+		SBorder::Construct(
+			SBorder::FArguments()
+			.BorderImage(FAppStyle::GetBrush("BlueprintEditor.PipelineSeparator"))
+			.Padding(0.0f)
+		);
+	}
+
+	// SWidget interface
+	virtual FVector2D ComputeDesiredSize(float) const override
+	{
+		constexpr float Height = 20.0f;
+		constexpr float Thickness = 16.0f;
+		return FVector2D(Thickness, Height);
+	}
+	// End of SWidget interface
+};
+
+//////////////////////////////////////////////////////////////////////////
+// FRenderPagesBlueprintEditorToolbar
+
+UE::RenderPages::Private::FRenderPagesBlueprintEditorToolbar::FRenderPagesBlueprintEditorToolbar(TSharedPtr<IRenderPageCollectionEditor>& InRenderPagesEditor)
+	: BlueprintEditorWeakPtr(InRenderPagesEditor)
+{}
+
+void UE::RenderPages::Private::FRenderPagesBlueprintEditorToolbar::AddRenderPagesBlueprintEditorModesToolbar(TSharedPtr<FExtender> Extender)
+{
+	if (const TSharedPtr<IRenderPageCollectionEditor> BlueprintEditor = BlueprintEditorWeakPtr.Pin())
+	{
+		Extender->AddToolBarExtension(
+			"Asset",
+			EExtensionHook::After,
+			BlueprintEditor->GetToolkitCommands(),
+			FToolBarExtensionDelegate::CreateSP(this, &FRenderPagesBlueprintEditorToolbar::FillRenderPagesBlueprintEditorModesToolbar));
+	}
+}
+
+void UE::RenderPages::Private::FRenderPagesBlueprintEditorToolbar::AddListingModeToolbar(UToolMenu* InMenu)
+{
+	FToolMenuSection& Section = InMenu->AddSection("Tools");
+
+	Section.InsertPosition = FToolMenuInsert("Asset", EToolMenuInsertType::After);
+
+	Section.AddEntry(FToolMenuEntry::InitToolBarButton(
+		FRenderPagesEditorCommands::Get().AddPage,
+		TAttribute<FText>(),
+		TAttribute<FText>(),
+		FSlateIcon(FAppStyle::GetAppStyleSetName(), "Icons.Plus")
+	));
+
+	Section.AddEntry(FToolMenuEntry::InitToolBarButton(
+		FRenderPagesEditorCommands::Get().CopyPage,
+		TAttribute<FText>(),
+		TAttribute<FText>(),
+		FSlateIcon(FAppStyle::GetAppStyleSetName(), "GenericCommands.Duplicate")
+	));
+
+	Section.AddEntry(FToolMenuEntry::InitToolBarButton(
+		FRenderPagesEditorCommands::Get().DeletePage,
+		TAttribute<FText>(),
+		TAttribute<FText>(),
+		FSlateIcon(FAppStyle::GetAppStyleSetName(), "Icons.Minus")
+	));
+}
+
+void UE::RenderPages::Private::FRenderPagesBlueprintEditorToolbar::AddLogicModeToolbar(UToolMenu* InMenu)
+{
+	FToolMenuSection& Section = InMenu->AddSection("Tools");
+
+	Section.InsertPosition = FToolMenuInsert("Asset", EToolMenuInsertType::After);
+
+	Section.AddEntry(FToolMenuEntry::InitToolBarButton(
+		FRenderPagesEditorCommands::Get().AutoCompileGraph,
+		TAttribute<FText>(),
+		TAttribute<FText>(),
+		FSlateIcon(FRenderPagesEditorStyle::Get().GetStyleSetName(), "RenderPages.AutoCompileGraph")
+	));
+}
+
+void UE::RenderPages::Private::FRenderPagesBlueprintEditorToolbar::FillRenderPagesBlueprintEditorModesToolbar(FToolBarBuilder& ToolbarBuilder)
+{
+	if (const TSharedPtr<IRenderPageCollectionEditor> BlueprintEditor = BlueprintEditorWeakPtr.Pin())
+	{
+		UBlueprint* BlueprintObj = BlueprintEditor->GetBlueprintObj();
+		if (!BlueprintObj || (!FBlueprintEditorUtils::IsLevelScriptBlueprint(BlueprintObj) && !FBlueprintEditorUtils::IsInterfaceBlueprint(BlueprintObj) && !BlueprintObj->bIsNewlyCreated))
+		{
+			TAttribute<FName> GetActiveMode(BlueprintEditor.ToSharedRef(), &FBlueprintEditor::GetCurrentMode);
+			FOnModeChangeRequested SetActiveMode = FOnModeChangeRequested::CreateSP(BlueprintEditor.ToSharedRef(), &FBlueprintEditor::SetCurrentMode);
+
+			// Left side padding
+			BlueprintEditor->AddToolbarWidget(SNew(SSpacer).Size(FVector2D(4.0f, 1.0f)));
+
+			BlueprintEditor->AddToolbarWidget(
+				SNew(SModeWidget, FRenderPagesApplicationModes::GetLocalizedMode(FRenderPagesApplicationModes::ListingMode), FRenderPagesApplicationModes::ListingMode)
+				.OnGetActiveMode(GetActiveMode)
+				.OnSetActiveMode(SetActiveMode)
+				.ToolTip(IDocumentation::Get()->CreateToolTip(
+					LOCTEXT("ListingModeButtonTooltip", "Switch to Blueprint Listing Mode"),
+					nullptr,
+					TEXT("Shared/Editors/BlueprintEditor"),
+					TEXT("ListingMode")))
+				.IconImage(FAppStyle::GetBrush("BTEditor.Graph.NewTask"))
+				.AddMetaData<FTagMetaData>(FTagMetaData(TEXT("ListingMode")))
+			);
+
+			BlueprintEditor->AddToolbarWidget(SNew(SSpacer).Size(FVector2D(10.0f, 1.0f)));
+
+			BlueprintEditor->AddToolbarWidget(
+				SNew(SModeWidget, FRenderPagesApplicationModes::GetLocalizedMode(FRenderPagesApplicationModes::LogicMode), FRenderPagesApplicationModes::LogicMode)
+				.OnGetActiveMode(GetActiveMode)
+				.OnSetActiveMode(SetActiveMode)
+				.CanBeSelected(BlueprintEditor.Get(), &FBlueprintEditor::IsEditingSingleBlueprint)
+				.ToolTip(IDocumentation::Get()->CreateToolTip(
+					LOCTEXT("LogicModeButtonTooltip", "Switch to Logic Editing Mode"),
+					nullptr,
+					TEXT("Shared/Editors/BlueprintEditor"),
+					TEXT("GraphMode")))
+				.ToolTipText(LOCTEXT("LogicModeButtonTooltip", "Switch to Logic Editing Mode"))
+				.IconImage(FAppStyle::GetBrush("Icons.Blueprint"))
+				.AddMetaData<FTagMetaData>(FTagMetaData(TEXT("LogicMode")))
+			);
+
+			// Right side padding
+			BlueprintEditor->AddToolbarWidget(SNew(SSpacer).Size(FVector2D(10.0f, 1.0f)));
+		}
+	}
+}
+
+
+#undef LOCTEXT_NAMESPACE
diff --git a/Engine/Plugins/Experimental/RenderPages/Source/RenderPagesEditor/Private/Toolkit/RenderPageCollectionEditorToolbar.h b/Engine/Plugins/Experimental/RenderPages/Source/RenderPagesEditor/Private/Toolkit/RenderPageCollectionEditorToolbar.h
new file mode 100644
--- /dev/null
+++ b/Engine/Plugins/Experimental/RenderPages/Source/RenderPagesEditor/Private/Toolkit/RenderPageCollectionEditorToolbar.h
@@ -0,0 +1,42 @@
+﻿// Copyright Epic Games, Inc. All Rights Reserved.
+
+#pragma once
+
+#include "CoreMinimal.h"
+
+
+class FExtender;
+class FToolBarBuilder;
+class UToolMenu;
+
+namespace UE::RenderPages
+{
+	class IRenderPageCollectionEditor;
+}
+
+
+namespace UE::RenderPages::Private
+{
+	/**
+	 * Handles all of the toolbar related construction for the render pages blueprint editor.
+	 */
+	class FRenderPagesBlueprintEditorToolbar : public TSharedFromThis<FRenderPagesBlueprintEditorToolbar>
+	{
+	public:
+		FRenderPagesBlueprintEditorToolbar(TSharedPtr<IRenderPageCollectionEditor>& InRenderPagesEditor);
+
+		/** Adds the mode-switch UI to the editor. */
+		void AddRenderPagesBlueprintEditorModesToolbar(TSharedPtr<FExtender> Extender);
+
+		/** Adds the toolbar for the listing mode to the editor. */
+		void AddListingModeToolbar(UToolMenu* InMenu);
+
+		/** Adds the toolbar for the logic mode to the editor. */
+		void AddLogicModeToolbar(UToolMenu* InMenu);
+
+	public:
+		/** Creates the mode-switch UI. */
+		void FillRenderPagesBlueprintEditorModesToolbar(FToolBarBuilder& ToolbarBuilder);
+		TWeakPtr<IRenderPageCollectionEditor> BlueprintEditorWeakPtr;
+	};
+}
diff --git a/Engine/Plugins/Experimental/RenderPages/Source/RenderPagesEditor/Private/UI/Components/SRenderPagesDragHandle.h b/Engine/Plugins/Experimental/RenderPages/Source/RenderPagesEditor/Private/UI/Components/SRenderPagesDragHandle.h
new file mode 100644
--- /dev/null
+++ b/Engine/Plugins/Experimental/RenderPages/Source/RenderPagesEditor/Private/UI/Components/SRenderPagesDragHandle.h
@@ -0,0 +1,81 @@
+// Copyright Epic Games, Inc. All Rights Reserved.
+
+#pragma once
+
+#include "Widgets/SCompoundWidget.h"
+#include "Widgets/Layout/SBox.h"
+#include "Widgets/Images/SImage.h"
+#include "Styling/CoreStyle.h"
+#include "Input/DragAndDrop.h"
+
+
+/**
+ * A reusable widget for drag and drop functionality.
+ */
+template<typename DragDropType>
+class SRenderPagesDragHandle : public SCompoundWidget
+{
+public:
+	using WidgetType = typename DragDropType::WidgetType;
+	using HeldItemType = typename DragDropType::HeldItemType;
+
+	SLATE_BEGIN_ARGS(SRenderPagesDragHandle) {}
+		SLATE_ARGUMENT(TSharedPtr<WidgetType>, Widget)
+	SLATE_END_ARGS()
+
+	void Construct(const FArguments& InArgs, HeldItemType InItem)
+	{
+		WidgetWeakPtr = InArgs._Widget;
+		Item = MoveTemp(InItem);
+		ChildSlot
+		[
+			SNew(SBox)
+			.WidthOverride(25.0f)
+			[
+				SNew(SHorizontalBox)
+				+ SHorizontalBox::Slot()
+				.Padding(5.0f, 0.0f)
+				.VAlign(VAlign_Center)
+				[
+					SNew(SImage)
+					.Image(FCoreStyle::Get().GetBrush("VerticalBoxDragIndicatorShort"))
+				]
+			]
+		];
+	}
+
+	virtual FReply OnMouseButtonDown(const FGeometry& MyGeometry, const FPointerEvent& MouseEvent) override
+	{
+		return FReply::Handled().DetectDrag(SharedThis(this), EKeys::LeftMouseButton);
+	}
+
+	virtual FReply OnDragDetected(const FGeometry& MyGeometry, const FPointerEvent& MouseEvent) override
+	{
+		if (MouseEvent.IsMouseButtonDown(EKeys::LeftMouseButton))
+		{
+			if (TSharedPtr<FDragDropOperation> DragDropOp = CreateDragDropOperation())
+			{
+				return FReply::Handled().BeginDragDrop(DragDropOp.ToSharedRef());
+			}
+		}
+		return FReply::Unhandled();
+	}
+
+	TSharedPtr<FDragDropOperation> CreateDragDropOperation()
+	{
+		if (const TSharedPtr<WidgetType> Widget = WidgetWeakPtr.Pin())
+		{
+			TSharedPtr<DragDropType> DragDropOperation = MakeShared<DragDropType>(Widget, Item);
+			DragDropOperation->Construct();
+			return DragDropOperation;
+		}
+		return nullptr;
+	}
+
+private:
+	/** Holds the widget to display when dragging. */
+	TWeakPtr<WidgetType> WidgetWeakPtr;
+
+	/** Holds the item being dragged. */
+	HeldItemType Item;
+};
diff --git a/Engine/Plugins/Experimental/RenderPages/Source/RenderPagesEditor/Private/UI/Components/SRenderPagesEditableTextBlock.cpp b/Engine/Plugins/Experimental/RenderPages/Source/RenderPagesEditor/Private/UI/Components/SRenderPagesEditableTextBlock.cpp
new file mode 100644
--- /dev/null
+++ b/Engine/Plugins/Experimental/RenderPages/Source/RenderPagesEditor/Private/UI/Components/SRenderPagesEditableTextBlock.cpp
@@ -0,0 +1,97 @@
+﻿// Copyright Epic Games, Inc. All Rights Reserved.
+
+#include "UI/Components/SRenderPagesEditableTextBlock.h"
+
+#include "EditorFontGlyphs.h"
+#include "PropertyCustomizationHelpers.h"
+#include "SlateOptMacros.h"
+#include "Styling/AppStyle.h"
+#include "Widgets/Input/SButton.h"
+
+#define LOCTEXT_NAMESPACE "SRenderPagesEditableTextBlock"
+
+
+void UE::RenderPages::Private::SRenderPagesEditableTextBlock::Tick(const FGeometry&, const double, const float)
+{
+	if (bNeedsRename)
+	{
+		if (TextBlock)
+		{
+			TextBlock->EnterEditingMode();
+		}
+		bNeedsRename = false;
+	}
+}
+
+BEGIN_SLATE_FUNCTION_BUILD_OPTIMIZATION
+
+void UE::RenderPages::Private::SRenderPagesEditableTextBlock::Construct(const FArguments& InArgs)
+{
+	bNeedsRename = false;
+	OnTextCommittedDelegate = InArgs._OnTextCommitted;
+
+	SetText(InArgs._Text);
+
+	ChildSlot
+	[
+		SNew(SHorizontalBox)
+		+ SHorizontalBox::Slot()
+		// Text block
+		.FillWidth(1.f)
+		.VAlign(VAlign_Fill)
+		.Padding(2.f)
+		[
+			SAssignNew(TextBlock, SInlineEditableTextBlock)
+			.Text(Text)
+			.OnTextCommitted(this, &SRenderPagesEditableTextBlock::OnTextBlockCommitted)
+		]
+		+ SHorizontalBox::Slot()
+		// Button
+		.AutoWidth()
+		.VAlign(VAlign_Center)
+		.HAlign(HAlign_Right)
+		[
+			SNew(SButton)
+			.ButtonStyle(FAppStyle::Get(), "FlatButton")
+			.OnClicked_Lambda([this]()
+			{
+				bNeedsRename = true;
+				return FReply::Handled();
+			})
+			[
+				SNew(STextBlock)
+				.Font(FAppStyle::Get().GetFontStyle("FontAwesome.10"))
+				.Text(FEditorFontGlyphs::Pencil /*fa-pencil*/)
+			]
+		]
+	];
+}
+
+END_SLATE_FUNCTION_BUILD_OPTIMIZATION
+
+void UE::RenderPages::Private::SRenderPagesEditableTextBlock::SetText(TAttribute<FText> InText)
+{
+	SetText(InText.IsSet() ? InText.Get() : FText());
+}
+
+void UE::RenderPages::Private::SRenderPagesEditableTextBlock::SetText(const FText& InText)
+{
+	Text = InText;
+
+	if (TextBlock.IsValid())
+	{
+		TextBlock->SetText(Text);
+	}
+}
+
+void UE::RenderPages::Private::SRenderPagesEditableTextBlock::OnTextBlockCommitted(const FText& InLabel, ETextCommit::Type InCommitInfo)
+{
+	SetText(InLabel);
+	if (OnTextCommittedDelegate.IsBound())
+	{
+		SetText(OnTextCommittedDelegate.Execute(InLabel, InCommitInfo));
+	}
+}
+
+
+#undef LOCTEXT_NAMESPACE
diff --git a/Engine/Plugins/Experimental/RenderPages/Source/RenderPagesEditor/Private/UI/Components/SRenderPagesEditableTextBlock.h b/Engine/Plugins/Experimental/RenderPages/Source/RenderPagesEditor/Private/UI/Components/SRenderPagesEditableTextBlock.h
new file mode 100644
--- /dev/null
+++ b/Engine/Plugins/Experimental/RenderPages/Source/RenderPagesEditor/Private/UI/Components/SRenderPagesEditableTextBlock.h
@@ -0,0 +1,58 @@
+﻿// Copyright Epic Games, Inc. All Rights Reserved.
+
+#pragma once
+
+#include "CoreMinimal.h"
+#include "Widgets/SCompoundWidget.h"
+#include "Widgets/Text/SInlineEditableTextBlock.h"
+
+
+namespace UE::RenderPages::Private
+{
+	DECLARE_DELEGATE_RetVal_TwoParams(FText, FOnRenderPagesEditableTextBlockTextCommitted, const FText&, ETextCommit::Type);
+
+	/**
+	 * A reusable text widget that can be modified by pressing a pencil button that's shown next to it.
+	 */
+	class SRenderPagesEditableTextBlock : public SCompoundWidget
+	{
+	public:
+		SLATE_BEGIN_ARGS(SRenderPagesEditableTextBlock) {}
+			SLATE_ATTRIBUTE(FText, Text)// The text displayed in this text block
+			SLATE_EVENT(FOnRenderPagesEditableTextBlockTextCommitted, OnTextCommitted)// Called when the text is committed
+		SLATE_END_ARGS()
+
+		virtual void Tick(const FGeometry&, const double, const float) override;
+		void Construct(const FArguments& InArgs);
+
+		/**
+		 * Gets the text of this text block
+		 *
+		 * @return	This text block's text string
+		 */
+		const FText& GetText() const { return Text; }
+
+	public:
+		/** Sets the text for this text block */
+		void SetText(TAttribute<FText> InText);
+
+		/** Sets the text for this text block */
+		void SetText(const FText& InText);
+
+	protected:
+		void OnTextBlockCommitted(const FText& InLabel, ETextCommit::Type InCommitInfo);
+
+	protected:
+		/** A reference to the text block. */
+		TSharedPtr<SInlineEditableTextBlock> TextBlock;
+
+		/** The current text in the text block (read-only). */
+		FText Text;
+
+		/** Whether the group needs to be renamed. (As requested by a click on the rename button) */
+		bool bNeedsRename;
+
+		/** Delegate to execute when editing mode text is committed. */
+		FOnRenderPagesEditableTextBlockTextCommitted OnTextCommittedDelegate;
+	};
+}
diff --git a/Engine/Plugins/Experimental/RenderPages/Source/RenderPagesEditor/Private/UI/Components/SRenderPagesFileSelectorTextBlock.cpp b/Engine/Plugins/Experimental/RenderPages/Source/RenderPagesEditor/Private/UI/Components/SRenderPagesFileSelectorTextBlock.cpp
new file mode 100644
--- /dev/null
+++ b/Engine/Plugins/Experimental/RenderPages/Source/RenderPagesEditor/Private/UI/Components/SRenderPagesFileSelectorTextBlock.cpp
@@ -0,0 +1,101 @@
+﻿// Copyright Epic Games, Inc. All Rights Reserved.
+
+#include "UI/Components/SRenderPagesFileSelectorTextBlock.h"
+
+#include "DesktopPlatformModule.h"
+#include "EditorFontGlyphs.h"
+#include "PropertyCustomizationHelpers.h"
+#include "SlateOptMacros.h"
+#include "Styling/AppStyle.h"
+#include "Widgets/Input/SButton.h"
+
+#define LOCTEXT_NAMESPACE "SRenderPagesFileSelectorTextBlock"
+
+
+BEGIN_SLATE_FUNCTION_BUILD_OPTIMIZATION
+
+void UE::RenderPages::Private::SRenderPagesFileSelectorTextBlock::Construct(const FArguments& InArgs)
+{
+	FolderPath = InArgs._FolderPath;
+	OnTextCommittedDelegate = InArgs._OnTextCommitted;
+
+	SetText(InArgs._Text);
+
+	ChildSlot
+	[
+		SNew(SHorizontalBox)
+		+ SHorizontalBox::Slot()
+		// Text block
+		.FillWidth(1.f)
+		.VAlign(VAlign_Fill)
+		.Padding(2.f)
+		[
+			SAssignNew(TextBlock, SInlineEditableTextBlock)
+			.IsEnabled(false)
+			.Text(Text)
+		]
+		+ SHorizontalBox::Slot()
+		// Button
+		.AutoWidth()
+		.VAlign(VAlign_Center)
+		.HAlign(HAlign_Right)
+		[
+			SNew(SButton)
+			.ButtonStyle(FAppStyle::Get(), "FlatButton")
+			.OnClicked(this, &SRenderPagesFileSelectorTextBlock::OnOpenDirectoryDialog)
+			[
+				SNew(STextBlock)
+				.Font(FAppStyle::Get().GetFontStyle("FontAwesome.10"))
+				.Text(FEditorFontGlyphs::Folder_Open /*fa-folder-open*/)
+			]
+		]
+	];
+}
+
+END_SLATE_FUNCTION_BUILD_OPTIMIZATION
+
+void UE::RenderPages::Private::SRenderPagesFileSelectorTextBlock::SetText(TAttribute<FText> InText)
+{
+	SetText(InText.IsSet() ? InText.Get() : FText());
+}
+
+void UE::RenderPages::Private::SRenderPagesFileSelectorTextBlock::SetText(const FText& InText)
+{
+	Text = InText;
+
+	if (TextBlock.IsValid())
+	{
+		TextBlock->SetText(Text);
+	}
+}
+
+FReply UE::RenderPages::Private::SRenderPagesFileSelectorTextBlock::OnOpenDirectoryDialog()
+{
+	if (IDesktopPlatform* DesktopPlatform = FDesktopPlatformModule::Get())
+	{
+		TSharedPtr<SWindow> ParentWindow = FSlateApplication::Get().FindWidgetWindow(AsShared());
+		void* ParentWindowHandle = (ParentWindow.IsValid() && ParentWindow->GetNativeWindow().IsValid()) ? ParentWindow->GetNativeWindow()->GetOSWindowHandle() : nullptr;
+
+		FString OutFolderPath;
+		const bool bFolderSelected = DesktopPlatform->OpenDirectoryDialog(
+			ParentWindowHandle,
+			LOCTEXT("Title_OpenDirectoryDialog", "Choose a directory").ToString(),
+			FolderPath.Get(FPaths::ProjectDir() / TEXT("Saved/MovieRenders/")),
+			OutFolderPath
+		);
+
+		if (bFolderSelected)
+		{
+			const FText NewText = FText::FromString(OutFolderPath);
+			SetText(NewText);
+			if (OnTextCommittedDelegate.IsBound())
+			{
+				SetText(OnTextCommittedDelegate.Execute(NewText, ETextCommit::Type::OnEnter));
+			}
+		}
+	}
+	return FReply::Handled();
+}
+
+
+#undef LOCTEXT_NAMESPACE
diff --git a/Engine/Plugins/Experimental/RenderPages/Source/RenderPagesEditor/Private/UI/Components/SRenderPagesFileSelectorTextBlock.h b/Engine/Plugins/Experimental/RenderPages/Source/RenderPagesEditor/Private/UI/Components/SRenderPagesFileSelectorTextBlock.h
new file mode 100644
--- /dev/null
+++ b/Engine/Plugins/Experimental/RenderPages/Source/RenderPagesEditor/Private/UI/Components/SRenderPagesFileSelectorTextBlock.h
@@ -0,0 +1,58 @@
+﻿// Copyright Epic Games, Inc. All Rights Reserved.
+
+#pragma once
+
+#include "CoreMinimal.h"
+#include "Widgets/SCompoundWidget.h"
+#include "Widgets/Text/SInlineEditableTextBlock.h"
+
+
+namespace UE::RenderPages::Private
+{
+	DECLARE_DELEGATE_RetVal_TwoParams(FText, FOnRenderPagesFileSelectorTextBlockTextCommitted, const FText&, ETextCommit::Type);
+
+	/**
+	 * A reusable filepath text widget that can be modified by pressing a folder button that's shown next to it, which will open up a file selector popup.
+	 */
+	class SRenderPagesFileSelectorTextBlock : public SCompoundWidget
+	{
+	public:
+		SLATE_BEGIN_ARGS(SRenderPagesFileSelectorTextBlock) {}
+			SLATE_ATTRIBUTE(FText, Text)// The text displayed in this text block
+			SLATE_ATTRIBUTE(FString, FolderPath)// The default path of the open-directory dialog
+			SLATE_EVENT(FOnRenderPagesFileSelectorTextBlockTextCommitted, OnTextCommitted)// Callback when the text is committed
+		SLATE_END_ARGS()
+
+		void Construct(const FArguments& InArgs);
+
+		/**
+		 * Gets the text of this text block
+		 *
+		 * @return	This text block's text string
+		 */
+		const FText& GetText() const { return Text; }
+
+	public:
+		/** Sets the text for this text block */
+		void SetText(TAttribute<FText> InText);
+
+		/** Sets the text for this text block */
+		void SetText(const FText& InText);
+
+	protected:
+		FReply OnOpenDirectoryDialog();
+
+	protected:
+		/** A reference to the text block. */
+		TSharedPtr<SInlineEditableTextBlock> TextBlock;
+
+		/** The current text in the text block (read-only). */
+		FText Text;
+
+		/** The default path of the open-directory dialog (read-only). */
+		TAttribute<FString> FolderPath;
+
+		/** Delegate to execute when editing mode text is committed. */
+		FOnRenderPagesFileSelectorTextBlockTextCommitted OnTextCommittedDelegate;
+	};
+}
diff --git a/Engine/Plugins/Experimental/RenderPages/Source/RenderPagesEditor/Private/UI/Components/SRenderPagesPageViewerFrameSlider.cpp b/Engine/Plugins/Experimental/RenderPages/Source/RenderPagesEditor/Private/UI/Components/SRenderPagesPageViewerFrameSlider.cpp
new file mode 100644
--- /dev/null
+++ b/Engine/Plugins/Experimental/RenderPages/Source/RenderPagesEditor/Private/UI/Components/SRenderPagesPageViewerFrameSlider.cpp
@@ -0,0 +1,198 @@
+// Copyright Epic Games, Inc. All Rights Reserved.
+
+#include "UI/Components/SRenderPagesPageViewerFrameSlider.h"
+#include "RenderPage/RenderPageCollection.h"
+#include "SlateOptMacros.h"
+#include "Widgets/Input/SSlider.h"
+#include "Widgets/Layout/SBox.h"
+#include "Widgets/Text/STextBlock.h"
+
+#define LOCTEXT_NAMESPACE "SRenderPagesPageViewerFrameSlider"
+
+
+float UE::RenderPages::Private::SRenderPagesPageViewerFrameSlider::FrameSliderValue = 0;
+
+
+BEGIN_SLATE_FUNCTION_BUILD_OPTIMIZATION
+
+void UE::RenderPages::Private::SRenderPagesPageViewerFrameSlider::Construct(const FArguments& InArgs)
+{
+	OnValueChanged = InArgs._OnValueChanged;
+	OnCaptureEnd = InArgs._OnCaptureEnd;
+
+	SAssignNew(FrameSlider, SSlider)
+		.IndentHandle(true)
+		.MouseUsesStep(false)
+		.StepSize(0.0001)
+		.MinValue(0)
+		.MaxValue(1)
+		.Value(FrameSliderValue)
+		.OnValueChanged(this, &SRenderPagesPageViewerFrameSlider::FrameSliderValueChanged)
+		.OnMouseCaptureEnd(this, &SRenderPagesPageViewerFrameSlider::FrameSliderValueChangedEnd)
+		.OnControllerCaptureEnd(this, &SRenderPagesPageViewerFrameSlider::FrameSliderValueChangedEnd);
+
+	SAssignNew(FrameSliderStartFrameText, STextBlock);
+	SAssignNew(FrameSliderSelectedFrameText, STextBlock);
+	SAssignNew(FrameSliderEndFrameText, STextBlock);
+
+	ChildSlot
+	[
+		SNew(SVerticalBox)
+
+		// slider
+		+ SVerticalBox::Slot()
+		.AutoHeight()
+		.Padding(0.0f, 0.0f, 0.0f, 0.0f)
+		[
+			SNew(SBox)
+			.HAlign(HAlign_Fill)
+			.VAlign(VAlign_Fill)
+			.HeightOverride(20.0f)
+			[
+				FrameSlider.ToSharedRef()
+			]
+		]
+
+		// slider text
+		+ SVerticalBox::Slot()
+		.AutoHeight()
+		.Padding(5.0f, -3.0f, 5.0f, 3.0f)
+		[
+			SNew(SHorizontalBox)
+
+			// start frame
+			+ SHorizontalBox::Slot()
+			.FillWidth(1.0f)
+			.HAlign(HAlign_Left)
+			[
+				FrameSliderStartFrameText.ToSharedRef()
+			]
+
+			// selected frame
+			+ SHorizontalBox::Slot()
+			.FillWidth(1.0f)
+			.HAlign(HAlign_Center)
+			[
+				FrameSliderSelectedFrameText.ToSharedRef()
+			]
+
+			// end frame
+			+ SHorizontalBox::Slot()
+			.FillWidth(1.0f)
+			.HAlign(HAlign_Right)
+			[
+				FrameSliderEndFrameText.ToSharedRef()
+			]
+		]
+	];
+}
+
+END_SLATE_FUNCTION_BUILD_OPTIMIZATION
+
+void UE::RenderPages::Private::SRenderPagesPageViewerFrameSlider::FrameSliderValueChanged(const float NewValue)
+{
+	FrameSliderValue = NewValue;
+	OnValueChanged.ExecuteIfBound(FrameSliderValue);
+}
+
+void UE::RenderPages::Private::SRenderPagesPageViewerFrameSlider::FrameSliderValueChangedEnd()
+{
+	OnCaptureEnd.ExecuteIfBound();
+}
+
+
+void UE::RenderPages::Private::SRenderPagesPageViewerFrameSlider::ClearFramesText()
+{
+	if (FrameSliderStartFrameText.IsValid())
+	{
+		FrameSliderStartFrameText->SetText(FText());
+	}
+	if (FrameSliderSelectedFrameText.IsValid())
+	{
+		FrameSliderSelectedFrameText->SetText(FText());
+	}
+	if (FrameSliderEndFrameText.IsValid())
+	{
+		FrameSliderEndFrameText->SetText(FText());
+	}
+}
+
+void UE::RenderPages::Private::SRenderPagesPageViewerFrameSlider::SetFramesText(const int32 StartFrame, const int32 SelectedFrame, const int32 EndFrame)
+{
+	if (FrameSliderStartFrameText.IsValid())
+	{
+		FrameSliderStartFrameText->SetText(FText::AsNumber(StartFrame));
+	}
+	if (FrameSliderSelectedFrameText.IsValid())
+	{
+		FrameSliderSelectedFrameText->SetText(FText::AsNumber(SelectedFrame));
+	}
+	if (FrameSliderEndFrameText.IsValid())
+	{
+		FrameSliderEndFrameText->SetText(FText::AsNumber(EndFrame));
+	}
+}
+
+bool UE::RenderPages::Private::SRenderPagesPageViewerFrameSlider::SetFramesText(URenderPage* Page)
+{
+	if (!IsValid(Page) || !FrameSlider.IsValid())
+	{
+		return false;
+	}
+
+	TOptional<int32> SelectedFrame = GetSelectedFrame(Page);
+	TOptional<int32> StartFrame = Page->GetStartFrame();
+	TOptional<int32> EndFrame = Page->GetEndFrame();
+	if (!SelectedFrame.IsSet() || !StartFrame.IsSet() || !EndFrame.IsSet() || (StartFrame.Get(0) >= EndFrame.Get(0)))
+	{
+		return false;
+	}
+
+	SetFramesText(
+		StartFrame.Get(0),
+		SelectedFrame.Get(0),
+		EndFrame.Get(0) - 1
+	);
+	return true;
+}
+
+TOptional<int32> UE::RenderPages::Private::SRenderPagesPageViewerFrameSlider::GetSelectedSequenceFrame(URenderPage* Page)
+{
+	if (!IsValid(Page) || !FrameSlider.IsValid())
+	{
+		return TOptional<int32>();
+	}
+
+	TOptional<int32> OptionalStartFrame = Page->GetSequenceStartFrame();
+	TOptional<int32> OptionalEndFrame = Page->GetSequenceEndFrame();
+	int32 StartFrame = OptionalStartFrame.Get(0);
+	int32 EndFrame = OptionalEndFrame.Get(0) - 1;
+	if (!OptionalStartFrame.IsSet() || !OptionalEndFrame.IsSet() || (StartFrame >= EndFrame))
+	{
+		return TOptional<int32>();
+	}
+
+	return FMath::RoundToInt(FMath::Lerp(static_cast<double>(StartFrame), static_cast<double>(EndFrame), FrameSlider->GetValue()));
+}
+
+TOptional<int32> UE::RenderPages::Private::SRenderPagesPageViewerFrameSlider::GetSelectedFrame(URenderPage* Page)
+{
+	if (!IsValid(Page) || !FrameSlider.IsValid())
+	{
+		return TOptional<int32>();
+	}
+
+	TOptional<int32> OptionalStartFrame = Page->GetStartFrame();
+	TOptional<int32> OptionalEndFrame = Page->GetEndFrame();
+	int32 StartFrame = OptionalStartFrame.Get(0);
+	int32 EndFrame = OptionalEndFrame.Get(0) - 1;
+	if (!OptionalStartFrame.IsSet() || !OptionalEndFrame.IsSet() || (StartFrame >= EndFrame))
+	{
+		return TOptional<int32>();
+	}
+
+	return FMath::RoundToInt(FMath::Lerp(static_cast<double>(StartFrame), static_cast<double>(EndFrame), FrameSlider->GetValue()));
+}
+
+
+#undef LOCTEXT_NAMESPACE
diff --git a/Engine/Plugins/Experimental/RenderPages/Source/RenderPagesEditor/Private/UI/Components/SRenderPagesPageViewerFrameSlider.h b/Engine/Plugins/Experimental/RenderPages/Source/RenderPagesEditor/Private/UI/Components/SRenderPagesPageViewerFrameSlider.h
new file mode 100644
--- /dev/null
+++ b/Engine/Plugins/Experimental/RenderPages/Source/RenderPagesEditor/Private/UI/Components/SRenderPagesPageViewerFrameSlider.h
@@ -0,0 +1,76 @@
+// Copyright Epic Games, Inc. All Rights Reserved.
+
+#pragma once
+
+#include "CoreMinimal.h"
+#include "Widgets/SCompoundWidget.h"
+#include "Framework/SlateDelegates.h"
+
+
+class SBox;
+class SSlider;
+class STextBlock;
+class URenderPage;
+
+
+namespace UE::RenderPages::Private
+{
+	/**
+	 * The frame slider widget for the page viewer widgets.
+	 */
+	class SRenderPagesPageViewerFrameSlider : public SCompoundWidget
+	{
+	public:
+		SLATE_BEGIN_ARGS(SRenderPagesPageViewerFrameSlider) {}
+			SLATE_EVENT(FOnFloatValueChanged, OnValueChanged)// Called when the value is changed by the slider
+			SLATE_EVENT(FSimpleDelegate, OnCaptureEnd)// Called when the capture ends
+		SLATE_END_ARGS()
+
+		void Construct(const FArguments& InArgs);
+
+		/** Returns the current value of the frame slider (a value from 0.0 to 1.0, based on how far the slider is to the right, left=0.0, right=1.0). */
+		float GetValue() { return FrameSliderValue; }
+
+		/** Clears the text of the start, current and end frame. */
+		void ClearFramesText();
+
+		/** Sets the text of the start, current and end frame. */
+		void SetFramesText(const int32 StartFrame, const int32 SelectedFrame, const int32 EndFrame);
+
+		/** Sets the text of the start, current and end frame according to the start and end frame of the given render page, figures out the selected frame with the current value of the frame slider. */
+		bool SetFramesText(URenderPage* Page);
+
+		/** Gets the selected sequence frame of the given render page, based on the current value of the frame slider. */
+		TOptional<int32> GetSelectedSequenceFrame(URenderPage* Page);
+		
+		/** Gets the selected frame (that will be output) of the given render page, based on the current value of the frame slider. */
+		TOptional<int32> GetSelectedFrame(URenderPage* Page);
+
+	private:
+		void FrameSliderValueChanged(const float NewValue);
+		void FrameSliderValueChangedEnd();
+
+	private:
+		/** The widget that allows the user to select what frame they'd like to see. */
+		TSharedPtr<SSlider> FrameSlider;
+
+	private:
+		/** The value of the frame slider. */
+		static float FrameSliderValue;
+
+		/** The text for the start frame under the slider. */
+		TSharedPtr<STextBlock> FrameSliderStartFrameText;
+
+		/** The text for the selected frame under the slider. */
+		TSharedPtr<STextBlock> FrameSliderSelectedFrameText;
+
+		/** The text for the end frame under the slider. */
+		TSharedPtr<STextBlock> FrameSliderEndFrameText;
+
+		/** Called when the value is changed by the slider. */
+		FOnFloatValueChanged OnValueChanged;
+
+		/** Called when the capture ends. */
+		FSimpleDelegate OnCaptureEnd;
+	};
+}
diff --git a/Engine/Plugins/Experimental/RenderPages/Source/RenderPagesEditor/Private/UI/Components/SRenderPagesRemoteControlEntity.cpp b/Engine/Plugins/Experimental/RenderPages/Source/RenderPagesEditor/Private/UI/Components/SRenderPagesRemoteControlEntity.cpp
new file mode 100644
--- /dev/null
+++ b/Engine/Plugins/Experimental/RenderPages/Source/RenderPagesEditor/Private/UI/Components/SRenderPagesRemoteControlEntity.cpp
@@ -0,0 +1,63 @@
+// Copyright Epic Games, Inc. All Rights Reserved.
+
+#include "SRenderPagesRemoteControlEntity.h"
+#include "RemoteControlBinding.h"
+#include "RemoteControlEntity.h"
+#include "RemoteControlPreset.h"
+#include "RemoteControlField.h"
+#include "Widgets/Layout/SBorder.h"
+#include "Widgets/SBoxPanel.h"
+#include "Widgets/Text/SInlineEditableTextBlock.h"
+
+#define LOCTEXT_NAMESPACE "SRenderPagesExposedEntity"
+
+
+TSharedPtr<FRemoteControlEntity> UE::RenderPages::Private::SRenderPagesRemoteControlEntity::GetEntity() const
+{
+	if (const TSharedPtr<FRemoteControlEntity> Entity = PresetWeakPtr->GetExposedEntity(EntityId).Pin())
+	{
+		return Entity;
+	}
+	return nullptr;
+}
+
+void UE::RenderPages::Private::SRenderPagesRemoteControlEntity::Initialize(const FGuid& InEntityId, URemoteControlPreset* InPreset)
+{
+	EntityId = InEntityId;
+	PresetWeakPtr = InPreset;
+
+	if (ensure(InPreset))
+	{
+		if (const TSharedPtr<FRemoteControlEntity> RCEntity = InPreset->GetExposedEntity(InEntityId).Pin())
+		{
+			CachedLabel = RCEntity->GetLabel();
+		}
+	}
+}
+
+TSharedRef<SWidget> UE::RenderPages::Private::SRenderPagesRemoteControlEntity::CreateEntityWidget(const TSharedPtr<SWidget> ValueWidget)
+{
+	FRenderPagesMakeNodeWidgetArgs Args;
+
+	TSharedRef<SBorder> Widget = SNew(SBorder)
+		.Padding(0.0f)
+		.BorderImage(new FSlateNoResource());
+
+	Args.NameWidget = SNew(SHorizontalBox)
+		.Clipping(EWidgetClipping::OnDemand)
+		+ SHorizontalBox::Slot()
+		.AutoWidth()
+		[
+			SNew(SInlineEditableTextBlock)
+			.Text(FText::FromName(CachedLabel))
+			.IsReadOnly_Lambda([this]() { return true; /*!bEditMode.Get();*/ })
+		];
+
+	Args.ValueWidget = ValueWidget;
+
+	Widget->SetContent(MakeNodeWidget(Args));
+	return Widget;
+}
+
+
+#undef LOCTEXT_NAMESPACE
diff --git a/Engine/Plugins/Experimental/RenderPages/Source/RenderPagesEditor/Private/UI/Components/SRenderPagesRemoteControlEntity.h b/Engine/Plugins/Experimental/RenderPages/Source/RenderPagesEditor/Private/UI/Components/SRenderPagesRemoteControlEntity.h
new file mode 100644
--- /dev/null
+++ b/Engine/Plugins/Experimental/RenderPages/Source/RenderPagesEditor/Private/UI/Components/SRenderPagesRemoteControlEntity.h
@@ -0,0 +1,41 @@
+// Copyright Epic Games, Inc. All Rights Reserved.
+
+#pragma once
+
+#include "SRenderPagesRemoteControlTreeNode.h"
+
+
+class SWidget;
+struct FRemoteControlEntity;
+class URemoteControlPreset;
+
+
+namespace UE::RenderPages::Private
+{
+	/**
+	 * The remote control entity widget, copied over from the remote control plugin, and slightly modified and cleaned up for usage in the render pages plugin.
+	 *
+	 * @Note if you inherit from this struct, you must call SRenderPagesExposedEntity::Initialize.
+	 */
+	struct SRenderPagesRemoteControlEntity : SRenderPagesRemoteControlTreeNode
+	{
+		//~ SRenderPagesTreeNode interface
+		TSharedPtr<FRemoteControlEntity> GetEntity() const;
+		virtual FGuid GetRCId() const override final { return EntityId; }
+		//~ End SRenderPagesTreeNode interface
+
+	protected:
+		void Initialize(const FGuid& InEntityId, URemoteControlPreset* InPreset);
+
+		/** Create an exposed entity widget. */
+		TSharedRef<SWidget> CreateEntityWidget(const TSharedPtr<SWidget> ValueWidget);
+
+	protected:
+		/** Id of the entity. */
+		FGuid EntityId;
+		/** The underlying preset. */
+		TWeakObjectPtr<URemoteControlPreset> PresetWeakPtr;
+		/** Display name of the entity. */
+		FName CachedLabel;
+	};
+}
diff --git a/Engine/Plugins/Experimental/RenderPages/Source/RenderPagesEditor/Private/UI/Components/SRenderPagesRemoteControlField.cpp b/Engine/Plugins/Experimental/RenderPages/Source/RenderPagesEditor/Private/UI/Components/SRenderPagesRemoteControlField.cpp
new file mode 100644
--- /dev/null
+++ b/Engine/Plugins/Experimental/RenderPages/Source/RenderPagesEditor/Private/UI/Components/SRenderPagesRemoteControlField.cpp
@@ -0,0 +1,182 @@
+// Copyright Epic Games, Inc. All Rights Reserved.
+
+#include "SRenderPagesRemoteControlField.h"
+#include "UI/Utils/RenderPagesWidgetUtils.h"
+#include "Algo/Transform.h"
+#include "IDetailTreeNode.h"
+#include "IPropertyRowGenerator.h"
+#include "Modules/ModuleManager.h"
+#include "RemoteControlField.h"
+#include "RemoteControlPreset.h"
+#include "SRenderPagesRemoteControlTreeNode.h"
+#include "UObject/Object.h"
+#include "Widgets/Text/SInlineEditableTextBlock.h"
+
+#define LOCTEXT_NAMESPACE "SRenderPagesExposedField"
+
+
+namespace ExposedFieldUtils
+{
+	TSharedRef<SWidget> CreateNodeValueWidget(const TSharedPtr<IDetailTreeNode>& Node)
+	{
+		FNodeWidgets NodeWidgets = Node->CreateNodeWidgets();
+
+		TSharedRef<SHorizontalBox> FieldWidget = SNew(SHorizontalBox);
+
+		if (NodeWidgets.ValueWidget)
+		{
+			FieldWidget->AddSlot()
+				.Padding(FMargin(3.0f, 2.0f))
+				.HAlign(HAlign_Right)
+				.FillWidth(1.0f)
+				[
+					NodeWidgets.ValueWidget.ToSharedRef()
+				];
+		}
+		else if (NodeWidgets.WholeRowWidget)
+		{
+			FieldWidget->AddSlot()
+				.Padding(FMargin(3.0f, 2.0f))
+				.FillWidth(1.0f)
+				[
+					NodeWidgets.WholeRowWidget.ToSharedRef()
+				];
+		}
+
+		return FieldWidget;
+	}
+}
+
+TSharedPtr<UE::RenderPages::Private::SRenderPagesRemoteControlTreeNode> UE::RenderPages::Private::SRenderPagesRemoteControlField::MakeInstance(const FRenderPagesRemoteControlGenerateWidgetArgs& Args)
+{
+	return SNew(SRenderPagesRemoteControlField, StaticCastSharedPtr<FRemoteControlField>(Args.Entity), Args.ColumnSizeData).Preset(Args.Preset);
+}
+
+void UE::RenderPages::Private::SRenderPagesRemoteControlField::Construct(const FArguments& InArgs, TWeakPtr<FRemoteControlField> InField, FRenderPagesRemoteControlColumnSizeData InColumnSizeData)
+{
+	FieldWeakPtr = MoveTemp(InField);
+	ColumnSizeData = MoveTemp(InColumnSizeData);
+
+	if (const TSharedPtr<FRemoteControlField> Field = FieldWeakPtr.Pin())
+	{
+		Initialize(Field->GetId(), InArgs._Preset.Get());
+
+		CachedLabel = Field->GetLabel();
+		EntityId = Field->GetId();
+
+		if (Field->FieldType == EExposedFieldType::Property)
+		{
+			ConstructPropertyWidget();
+		}
+	}
+}
+
+void UE::RenderPages::Private::SRenderPagesRemoteControlField::GetNodeChildren(TArray<TSharedPtr<SRenderPagesRemoteControlTreeNode>>& OutChildren) const
+{
+	OutChildren.Append(ChildWidgets);
+}
+
+UE::RenderPages::Private::SRenderPagesRemoteControlTreeNode::ENodeType UE::RenderPages::Private::SRenderPagesRemoteControlField::GetRCType() const
+{
+	return ENodeType::Field;
+}
+
+FName UE::RenderPages::Private::SRenderPagesRemoteControlField::GetFieldLabel() const
+{
+	return CachedLabel;
+}
+
+EExposedFieldType UE::RenderPages::Private::SRenderPagesRemoteControlField::GetFieldType() const
+{
+	if (const TSharedPtr<FRemoteControlField> Field = FieldWeakPtr.Pin())
+	{
+		return Field->FieldType;
+	}
+	return EExposedFieldType::Invalid;
+}
+
+void UE::RenderPages::Private::SRenderPagesRemoteControlField::Refresh()
+{
+	if (const TSharedPtr<FRemoteControlField> Field = FieldWeakPtr.Pin())
+	{
+		CachedLabel = Field->GetLabel();
+
+		if (Field->FieldType == EExposedFieldType::Property)
+		{
+			ConstructPropertyWidget();
+		}
+	}
+}
+
+void UE::RenderPages::Private::SRenderPagesRemoteControlField::GetBoundObjects(TSet<UObject*>& OutBoundObjects) const
+{
+	if (const TSharedPtr<FRemoteControlField> Field = FieldWeakPtr.Pin())
+	{
+		OutBoundObjects.Append(Field->GetBoundObjects());
+	}
+}
+
+TSharedRef<SWidget> UE::RenderPages::Private::SRenderPagesRemoteControlField::ConstructWidget()
+{
+	if (const TSharedPtr<FRemoteControlField> Field = FieldWeakPtr.Pin())
+	{
+		// For the moment, just use the first object.
+		TArray<UObject*> Objects = Field->GetBoundObjects();
+		if ((GetFieldType() == EExposedFieldType::Property) && (Objects.Num() > 0))
+		{
+			FPropertyRowGeneratorArgs Args;
+			Generator = FModuleManager::GetModuleChecked<FPropertyEditorModule>("PropertyEditor").CreatePropertyRowGenerator(Args);
+			Generator->SetObjects({Objects[0]});
+
+			if (TSharedPtr<IDetailTreeNode> Node = RenderPagesWidgetUtils::FindNode(Generator->GetRootTreeNodes(), Field->FieldPathInfo.ToPathPropertyString(), RenderPagesWidgetUtils::ERenderPagesFindNodeMethod::Path))
+			{
+				TArray<TSharedRef<IDetailTreeNode>> ChildNodes;
+				Node->GetChildren(ChildNodes);
+				ChildWidgets.Reset(ChildNodes.Num());
+
+				for (const TSharedRef<IDetailTreeNode>& ChildNode : ChildNodes)
+				{
+					ChildWidgets.Add(SNew(SRenderPagesRemoteControlFieldChildNode, ChildNode, ColumnSizeData));
+				}
+
+				return MakeFieldWidget(ExposedFieldUtils::CreateNodeValueWidget(MoveTemp(Node)));
+			}
+		}
+	}
+
+	return MakeFieldWidget(SNullWidget::NullWidget);
+}
+
+TSharedRef<SWidget> UE::RenderPages::Private::SRenderPagesRemoteControlField::MakeFieldWidget(const TSharedRef<SWidget>& InWidget)
+{
+	return CreateEntityWidget(InWidget);
+}
+
+void UE::RenderPages::Private::SRenderPagesRemoteControlField::ConstructPropertyWidget()
+{
+	ChildSlot.AttachWidget(ConstructWidget());
+}
+
+
+void UE::RenderPages::Private::SRenderPagesRemoteControlFieldChildNode::Construct(const FArguments& InArgs, const TSharedRef<IDetailTreeNode>& InNode, FRenderPagesRemoteControlColumnSizeData InColumnSizeData)
+{
+	TArray<TSharedRef<IDetailTreeNode>> ChildNodes;
+	InNode->GetChildren(ChildNodes);
+
+	Algo::Transform(ChildNodes, ChildrenNodes, [InColumnSizeData](const TSharedRef<IDetailTreeNode>& ChildNode) { return SNew(SRenderPagesRemoteControlFieldChildNode, ChildNode, InColumnSizeData); });
+
+	ColumnSizeData = InColumnSizeData;
+
+	FNodeWidgets Widgets = InNode->CreateNodeWidgets();
+	FRenderPagesMakeNodeWidgetArgs Args;
+	Args.NameWidget = Widgets.NameWidget;
+	Args.ValueWidget = ExposedFieldUtils::CreateNodeValueWidget(InNode);
+
+	ChildSlot
+	[
+		MakeNodeWidget(Args)
+	];
+}
+
+
+#undef LOCTEXT_NAMESPACE /*RemoteControlPanel*/
diff --git a/Engine/Plugins/Experimental/RenderPages/Source/RenderPagesEditor/Private/UI/Components/SRenderPagesRemoteControlField.h b/Engine/Plugins/Experimental/RenderPages/Source/RenderPagesEditor/Private/UI/Components/SRenderPagesRemoteControlField.h
new file mode 100644
--- /dev/null
+++ b/Engine/Plugins/Experimental/RenderPages/Source/RenderPagesEditor/Private/UI/Components/SRenderPagesRemoteControlField.h
@@ -0,0 +1,91 @@
+// Copyright Epic Games, Inc. All Rights Reserved.
+
+#pragma once
+
+#include "CoreMinimal.h"
+#include "RemoteControlField.h"
+#include "UI/Components/SRenderPagesRemoteControlEntity.h"
+#include "UI/Components/SRenderPagesRemoteControlTreeNode.h"
+#include "Widgets/DeclarativeSyntaxSupport.h"
+
+
+class IDetailTreeNode;
+class URemoteControlPreset;
+enum class EExposedFieldType : uint8;
+struct FGuid;
+struct FRemoteControlField;
+struct FRenderPagesGenerateWidgetArgs;
+
+
+namespace UE::RenderPages::Private
+{
+	/**
+	 * The remote control field child node widget, copied over from the remote control plugin, and slightly modified and cleaned up for usage in the render pages plugin.
+	 *
+	 * Represents a child of an exposed field widget.
+	 */
+	struct SRenderPagesRemoteControlFieldChildNode : SRenderPagesRemoteControlTreeNode
+	{
+		SLATE_BEGIN_ARGS(SRenderPagesRemoteControlFieldChildNode) {}
+		SLATE_END_ARGS()
+
+		void Construct(const FArguments& InArgs, const TSharedRef<IDetailTreeNode>& InNode, FRenderPagesRemoteControlColumnSizeData InColumnSizeData);
+		virtual void GetNodeChildren(TArray<TSharedPtr<SRenderPagesRemoteControlTreeNode>>& OutChildren) const override { return OutChildren.Append(ChildrenNodes); }
+		virtual FGuid GetRCId() const override { return FGuid(); }
+		virtual ENodeType GetRCType() const override { return ENodeType::FieldChild; }
+
+		TArray<TSharedPtr<SRenderPagesRemoteControlFieldChildNode>> ChildrenNodes;
+	};
+
+
+	/**
+	 * The remote control field widget, copied over from the remote control plugin, and slightly modified and cleaned up for usage in the render pages plugin.
+	 *
+	 * Widget that displays an exposed field.
+	 */
+	struct SRenderPagesRemoteControlField : SRenderPagesRemoteControlEntity
+	{
+		SLATE_BEGIN_ARGS(SRenderPagesRemoteControlField)
+				: _Preset(nullptr) {}
+			SLATE_ATTRIBUTE(URemoteControlPreset*, Preset)
+		SLATE_END_ARGS()
+
+		static TSharedPtr<SRenderPagesRemoteControlTreeNode> MakeInstance(const FRenderPagesRemoteControlGenerateWidgetArgs& Args);
+
+		void Construct(const FArguments& InArgs, TWeakPtr<FRemoteControlField> Field, FRenderPagesRemoteControlColumnSizeData ColumnSizeData);
+
+		//~ SRenderPagesTreeNode Interface 
+		virtual void GetNodeChildren(TArray<TSharedPtr<SRenderPagesRemoteControlTreeNode>>& OutChildren) const override;
+		virtual ENodeType GetRCType() const override;
+		virtual void Refresh() override;
+		//~ End SRenderPagesTreeNode Interface
+
+		/** Get a weak pointer to the underlying remote control field. */
+		TWeakPtr<FRemoteControlField> GetRemoteControlField() const { return FieldWeakPtr; }
+
+		/** Get this field's label. */
+		FName GetFieldLabel() const;
+
+		/** Get this field's type. */
+		EExposedFieldType GetFieldType() const;
+
+
+		/** Returns this widget's underlying objects. */
+		void GetBoundObjects(TSet<UObject*>& OutBoundObjects) const;
+
+	private:
+		/** Construct a property widget. */
+		TSharedRef<SWidget> ConstructWidget();
+		/** Create the wrapper around the field value widget. */
+		TSharedRef<SWidget> MakeFieldWidget(const TSharedRef<SWidget>& InWidget);
+		/** Construct this field widget as a property widget. */
+		void ConstructPropertyWidget();
+	private:
+		/** Weak pointer to the underlying RC Field. */
+		TWeakPtr<FRemoteControlField> FieldWeakPtr;
+		/** This exposed field's child widgets (ie. An array's rows) */
+		TArray<TSharedPtr<SRenderPagesRemoteControlFieldChildNode>> ChildWidgets;
+		/** The property row generator. */
+		TSharedPtr<IPropertyRowGenerator> Generator;
+	};
+}
diff --git a/Engine/Plugins/Experimental/RenderPages/Source/RenderPagesEditor/Private/UI/Components/SRenderPagesRemoteControlTreeNode.cpp b/Engine/Plugins/Experimental/RenderPages/Source/RenderPagesEditor/Private/UI/Components/SRenderPagesRemoteControlTreeNode.cpp
new file mode 100644
--- /dev/null
+++ b/Engine/Plugins/Experimental/RenderPages/Source/RenderPagesEditor/Private/UI/Components/SRenderPagesRemoteControlTreeNode.cpp
@@ -0,0 +1,93 @@
+// Copyright Epic Games, Inc. All Rights Reserved.
+
+#include "SRenderPagesRemoteControlTreeNode.h"
+#include "Styling/AppStyle.h"
+#include "Widgets/Layout/SBox.h"
+#include "Widgets/SBoxPanel.h"
+#include "Widgets/SOverlay.h"
+
+#define LOCTEXT_NAMESPACE "SRenderPagesTreeNode"
+
+
+TSharedRef<SWidget> UE::RenderPages::Private::SRenderPagesRemoteControlTreeNode::MakeSplitRow(TSharedRef<SWidget> LeftColumn, TSharedRef<SWidget> RightColumn)
+{
+	TAttribute<float> LeftColumnAttribute;
+	LeftColumnAttribute.BindRaw(this, &SRenderPagesRemoteControlTreeNode::GetLeftColumnWidth);
+	TAttribute<float> RightColumnAttribute;
+	RightColumnAttribute.BindRaw(this, &SRenderPagesRemoteControlTreeNode::GetRightColumnWidth);
+
+	return SNew(SSplitter)
+		.Style(FAppStyle::Get(), "DetailsView.Splitter")
+		.PhysicalSplitterHandleSize(1.0f)
+		.HitDetectionSplitterHandleSize(5.0f)
+		+ SSplitter::Slot()
+		.Value(MoveTemp(LeftColumnAttribute))
+		.OnSlotResized(SSplitter::FOnSlotResized::CreateRaw(this, &SRenderPagesRemoteControlTreeNode::OnLeftColumnResized))
+		[
+			LeftColumn
+		]
+		+ SSplitter::Slot()
+		.Value(MoveTemp(RightColumnAttribute))
+		.OnSlotResized(SSplitter::FOnSlotResized::CreateRaw(this, &SRenderPagesRemoteControlTreeNode::SetColumnWidth))
+		[
+			RightColumn
+		];
+}
+
+TSharedRef<SWidget> UE::RenderPages::Private::SRenderPagesRemoteControlTreeNode::MakeNodeWidget(const FRenderPagesMakeNodeWidgetArgs& Args)
+{
+	auto WidgetOrNull = [](const TSharedPtr<SWidget>& Widget) { return Widget ? Widget.ToSharedRef() : SNullWidget::NullWidget; };
+
+	TSharedRef<SWidget> LeftColumn =
+		SNew(SHorizontalBox)
+		.Clipping(EWidgetClipping::OnDemand)
+		// Name
+		+ SHorizontalBox::Slot()
+		.VAlign(VAlign_Center)
+		.AutoWidth()
+		[
+			WidgetOrNull(Args.NameWidget)
+		];
+
+	TSharedRef<SWidget> RightColumn =
+		SNew(SOverlay)
+		// Value
+		+ SOverlay::Slot()
+		.HAlign(HAlign_Fill)
+		[
+			SNew(SBox)
+			.HAlign(HAlign_Left)
+			.VAlign(VAlign_Center)
+			[
+				WidgetOrNull(Args.ValueWidget)
+			]
+		];
+
+	return MakeSplitRow(LeftColumn, RightColumn);
+}
+
+void UE::RenderPages::Private::SRenderPagesRemoteControlTreeNode::OnLeftColumnResized(float) const
+{
+	// This has to be bound or the splitter will take it upon itself to determine the size
+	// We do nothing here because it is handled by the column size data
+}
+
+float UE::RenderPages::Private::SRenderPagesRemoteControlTreeNode::GetLeftColumnWidth() const
+{
+	const float Offset = (GetRCType() == ENodeType::Group) ? SplitterOffset : 0;
+	return FMath::Clamp(ColumnSizeData.LeftColumnWidth.Get() + Offset, 0.f, 1.f);
+}
+
+float UE::RenderPages::Private::SRenderPagesRemoteControlTreeNode::GetRightColumnWidth() const
+{
+	const float Offset = (GetRCType() == ENodeType::Group) ? SplitterOffset : 0;
+	return FMath::Clamp(ColumnSizeData.RightColumnWidth.Get() - Offset, 0.f, 1.f);
+}
+
+void UE::RenderPages::Private::SRenderPagesRemoteControlTreeNode::SetColumnWidth(float InWidth)
+{
+	ColumnSizeData.SetColumnWidth(FMath::Clamp(InWidth + SplitterOffset, 0.f, 1.f));
+}
+
+
+#undef LOCTEXT_NAMESPACE /*RemoteControlPanelNode*/
diff --git a/Engine/Plugins/Experimental/RenderPages/Source/RenderPagesEditor/Private/UI/Components/SRenderPagesRemoteControlTreeNode.h b/Engine/Plugins/Experimental/RenderPages/Source/RenderPagesEditor/Private/UI/Components/SRenderPagesRemoteControlTreeNode.h
new file mode 100644
--- /dev/null
+++ b/Engine/Plugins/Experimental/RenderPages/Source/RenderPagesEditor/Private/UI/Components/SRenderPagesRemoteControlTreeNode.h
@@ -0,0 +1,103 @@
+// Copyright Epic Games, Inc. All Rights Reserved.
+
+#pragma once
+
+#include "Templates/SharedPointer.h"
+#include "Containers/Array.h"
+#include "Misc/Guid.h"
+#include "Widgets/Layout/SSplitter.h"
+#include "Widgets/SCompoundWidget.h"
+
+
+class SWidget;
+class URemoteControlPreset;
+struct FRemoteControlEntity;
+
+
+namespace UE::RenderPages::Private
+{
+	/**
+	 * The remote control column size data struct, copied over from the remote control plugin, and slightly modified and cleaned up for usage in the render pages plugin.
+	 */
+	struct FRenderPagesRemoteControlColumnSizeData
+	{
+		TAttribute<float> LeftColumnWidth;
+		TAttribute<float> RightColumnWidth;
+		SSplitter::FOnSlotResized OnWidthChanged;
+
+		void SetColumnWidth(const float InWidth)
+		{
+			OnWidthChanged.ExecuteIfBound(InWidth);
+		}
+	};
+
+	/**
+	 * The remote control generate widget args struct, copied over from the remote control plugin, and slightly modified and cleaned up for usage in the render pages plugin.
+	 */
+	struct FRenderPagesRemoteControlGenerateWidgetArgs
+	{
+		URemoteControlPreset* Preset = nullptr;
+		FRenderPagesRemoteControlColumnSizeData ColumnSizeData;
+		TSharedPtr<FRemoteControlEntity> Entity;
+	};
+
+
+	/**
+	 * The remote control tree node widget, copied over from the remote control plugin, and slightly modified and cleaned up for usage in the render pages plugin.
+	 *
+	 * A node in the panel tree view.
+	 */
+	struct SRenderPagesRemoteControlTreeNode : SCompoundWidget
+	{
+		enum class ENodeType : uint8
+		{
+			Invalid,
+			Group,
+			Field,
+			FieldChild,
+			Actor,
+			Material
+		};
+
+		/** Get this tree node's children. */
+		virtual void GetNodeChildren(TArray<TSharedPtr<SRenderPagesRemoteControlTreeNode>>& OutChildren) const {}
+		/** Get this node's ID if any. */
+		virtual FGuid GetRCId() const { return FGuid(); }
+		/** Get get this node's type. */
+		virtual ENodeType GetRCType() const { return ENodeType::Invalid; }
+		/** Refresh the node. */
+		virtual void Refresh() {}
+
+	protected:
+		struct FRenderPagesMakeNodeWidgetArgs
+		{
+			TSharedPtr<SWidget> DragHandle;
+			TSharedPtr<SWidget> NameWidget;
+			TSharedPtr<SWidget> RenameButton;
+			TSharedPtr<SWidget> ValueWidget;
+			TSharedPtr<SWidget> UnexposeButton;
+		};
+
+		/** Create a widget that represents a row with a splitter. */
+		TSharedRef<SWidget> MakeSplitRow(TSharedRef<SWidget> LeftColumn, TSharedRef<SWidget> RightColumn);
+		/** Create a widget that represents a node in the panel tree hierarchy. */
+		TSharedRef<SWidget> MakeNodeWidget(const FRenderPagesMakeNodeWidgetArgs& Args);
+
+	private:
+		/** Stub handler for column resize callback to prevent the splitter from handling it internally.  */
+		void OnLeftColumnResized(float) const;
+
+		//~ Wrappers around ColumnSizeData's delegate needed in order to offset the splitter for RC Groups. 
+		float GetLeftColumnWidth() const;
+		float GetRightColumnWidth() const;
+		void SetColumnWidth(float InWidth);
+
+	protected:
+		/** Holds the row's columns' width. */
+		FRenderPagesRemoteControlColumnSizeData ColumnSizeData;
+
+	private:
+		/** The splitter offset to align the group splitter with the other row's splitters. */
+		static constexpr float SplitterOffset = 0.008f;
+	};
+}
diff --git a/Engine/Plugins/Experimental/RenderPages/Source/RenderPagesEditor/Private/UI/SRenderPagesCollection.cpp b/Engine/Plugins/Experimental/RenderPages/Source/RenderPagesEditor/Private/UI/SRenderPagesCollection.cpp
new file mode 100644
--- /dev/null
+++ b/Engine/Plugins/Experimental/RenderPages/Source/RenderPagesEditor/Private/UI/SRenderPagesCollection.cpp
@@ -0,0 +1,68 @@
+// Copyright Epic Games, Inc. All Rights Reserved.
+
+#include "UI/SRenderPagesCollection.h"
+#include "Blueprints/RenderPagesBlueprint.h"
+#include "IRenderPageCollectionEditor.h"
+#include "RenderPage/RenderPageCollection.h"
+#include "Kismet2/BlueprintEditorUtils.h"
+#include "SlateOptMacros.h"
+#include "Widgets/Layout/SScrollBox.h"
+
+#define LOCTEXT_NAMESPACE "SRenderPagesCollection"
+
+
+void UE::RenderPages::Private::SRenderPagesCollection::Tick(const FGeometry&, const double, const float)
+{
+	if(DetailsView.IsValid())
+	{
+		if (const TSharedPtr<IRenderPageCollectionEditor> BlueprintEditor = BlueprintEditorWeakPtr.Pin())
+		{
+			if (URenderPageCollection* Collection = BlueprintEditor->GetInstance(); IsValid(Collection))
+			{
+				if (DetailsViewRenderPageCollectionWeakPtr != Collection)
+				{
+					DetailsViewRenderPageCollectionWeakPtr = Collection;
+					DetailsView->SetObject(Collection);
+				}
+			}
+		}
+	}
+}
+
+BEGIN_SLATE_FUNCTION_BUILD_OPTIMIZATION
+
+void UE::RenderPages::Private::SRenderPagesCollection::Construct(const FArguments& InArgs, TSharedPtr<IRenderPageCollectionEditor> InBlueprintEditor)
+{
+	BlueprintEditorWeakPtr = InBlueprintEditor;
+	DetailsViewRenderPageCollectionWeakPtr = InBlueprintEditor->GetInstance();
+
+	FPropertyEditorModule& PropertyEditorModule = FModuleManager::GetModuleChecked<FPropertyEditorModule>("PropertyEditor");
+	FDetailsViewArgs DetailsViewArgs;
+	DetailsViewArgs.bAllowSearch = false;
+	DetailsViewArgs.NameAreaSettings = FDetailsViewArgs::HideNameArea;
+	DetailsViewArgs.DefaultsOnlyVisibility = EEditDefaultsOnlyNodeVisibility::Hide;
+	DetailsViewArgs.NotifyHook = this;
+	DetailsView = PropertyEditorModule.CreateDetailView(DetailsViewArgs);
+	DetailsView->SetObject(DetailsViewRenderPageCollectionWeakPtr.Get());
+
+	ChildSlot
+	[
+		DetailsView->AsShared()
+	];
+}
+
+END_SLATE_FUNCTION_BUILD_OPTIMIZATION
+
+void UE::RenderPages::Private::SRenderPagesCollection::NotifyPreChange(FEditPropertyChain* PropertyAboutToChange) {}
+
+void UE::RenderPages::Private::SRenderPagesCollection::NotifyPostChange(const FPropertyChangedEvent& PropertyChangedEvent, FEditPropertyChain* PropertyThatChanged)
+{
+	if (const TSharedPtr<IRenderPageCollectionEditor> BlueprintEditor = BlueprintEditorWeakPtr.Pin())
+	{
+		BlueprintEditor->MarkAsModified();
+		BlueprintEditor->OnRenderPagesChanged().Broadcast();
+	}
+}
+
+
+#undef LOCTEXT_NAMESPACE
diff --git a/Engine/Plugins/Experimental/RenderPages/Source/RenderPagesEditor/Private/UI/SRenderPagesCollection.h b/Engine/Plugins/Experimental/RenderPages/Source/RenderPagesEditor/Private/UI/SRenderPagesCollection.h
new file mode 100644
--- /dev/null
+++ b/Engine/Plugins/Experimental/RenderPages/Source/RenderPagesEditor/Private/UI/SRenderPagesCollection.h
@@ -0,0 +1,49 @@
+// Copyright Epic Games, Inc. All Rights Reserved.
+
+#pragma once
+
+#include "CoreMinimal.h"
+#include "Widgets/SCompoundWidget.h"
+#include "Misc/NotifyHook.h"
+
+
+class URenderPageCollection;
+class IDetailsView;
+
+namespace UE::RenderPages
+{
+	class IRenderPageCollectionEditor;
+}
+
+
+namespace UE::RenderPages::Private
+{
+	/**
+	 * A widget with which the user can modify the render pages collection.
+	 * Doesn't contain any UI elements to modify the pages the collection contains.
+	 */
+	class SRenderPagesCollection : public SCompoundWidget, public FNotifyHook
+	{
+	public:
+		SLATE_BEGIN_ARGS(SRenderPagesCollection) {}
+		SLATE_END_ARGS()
+
+		virtual void Tick(const FGeometry&, const double, const float) override;
+		void Construct(const FArguments& InArgs, TSharedPtr<IRenderPageCollectionEditor> InBlueprintEditor);
+
+		//~ Begin FNotifyHook Interface
+		virtual void NotifyPreChange(FEditPropertyChain* PropertyAboutToChange) override;
+		virtual void NotifyPostChange(const FPropertyChangedEvent& PropertyChangedEvent, FEditPropertyChain* PropertyThatChanged) override;
+		//~ End FNotifyHook Interface
+
+	private:
+		/** A reference to the BP Editor that owns this collection. */
+		TWeakPtr<IRenderPageCollectionEditor> BlueprintEditorWeakPtr;
+
+		/** A reference to the details view. */
+		TSharedPtr<IDetailsView> DetailsView;
+
+		/** The render page collection that's being edited in the details view. */
+		TWeakObjectPtr<URenderPageCollection> DetailsViewRenderPageCollectionWeakPtr;
+	};
+}
diff --git a/Engine/Plugins/Experimental/RenderPages/Source/RenderPagesEditor/Private/UI/SRenderPagesPage.cpp b/Engine/Plugins/Experimental/RenderPages/Source/RenderPagesEditor/Private/UI/SRenderPagesPage.cpp
new file mode 100644
--- /dev/null
+++ b/Engine/Plugins/Experimental/RenderPages/Source/RenderPagesEditor/Private/UI/SRenderPagesPage.cpp
@@ -0,0 +1,66 @@
+// Copyright Epic Games, Inc. All Rights Reserved.
+
+#include "UI/SRenderPagesPage.h"
+#include "RenderPage/RenderPageCollection.h"
+#include "IRenderPageCollectionEditor.h"
+#include "PropertyEditorModule.h"
+#include "IDetailsView.h"
+#include "SlateOptMacros.h"
+
+#define LOCTEXT_NAMESPACE "SRenderPagesPage"
+
+
+BEGIN_SLATE_FUNCTION_BUILD_OPTIMIZATION
+
+void UE::RenderPages::Private::SRenderPagesPage::Construct(const FArguments& InArgs, TSharedPtr<IRenderPageCollectionEditor> InBlueprintEditor)
+{
+	BlueprintEditorWeakPtr = InBlueprintEditor;
+
+	FPropertyEditorModule& PropertyEditorModule = FModuleManager::GetModuleChecked<FPropertyEditorModule>("PropertyEditor");
+	FDetailsViewArgs DetailsViewArgs;
+	DetailsViewArgs.bAllowSearch = false;
+	DetailsViewArgs.NameAreaSettings = FDetailsViewArgs::HideNameArea;
+	DetailsViewArgs.NotifyHook = this;
+	DetailsView = PropertyEditorModule.CreateDetailView(DetailsViewArgs);
+	DetailsView->SetObject(nullptr);
+
+	HandlePagesSelectionChanged();
+	InBlueprintEditor->OnRenderPagesSelectionChanged().AddSP(this, &SRenderPagesPage::HandlePagesSelectionChanged);
+
+	ChildSlot
+	[
+		DetailsView->AsShared()
+	];
+}
+
+END_SLATE_FUNCTION_BUILD_OPTIMIZATION
+
+void UE::RenderPages::Private::SRenderPagesPage::NotifyPreChange(FEditPropertyChain* PropertyAboutToChange) {}
+
+void UE::RenderPages::Private::SRenderPagesPage::NotifyPostChange(const FPropertyChangedEvent& PropertyChangedEvent, FEditPropertyChain* PropertyThatChanged)
+{
+	if (const TSharedPtr<IRenderPageCollectionEditor> BlueprintEditor = BlueprintEditorWeakPtr.Pin())
+	{
+		BlueprintEditor->MarkAsModified();
+		BlueprintEditor->OnRenderPagesChanged().Broadcast();
+	}
+}
+
+void UE::RenderPages::Private::SRenderPagesPage::HandlePagesSelectionChanged()
+{
+	if (DetailsView.IsValid())
+	{
+		if (const TSharedPtr<IRenderPageCollectionEditor> BlueprintEditor = BlueprintEditorWeakPtr.Pin())
+		{
+			TArray<TWeakObjectPtr<UObject>> WeakSelectedPages;
+			if (const TArray<URenderPage*> SelectedPages = BlueprintEditor->GetSelectedRenderPages(); (SelectedPages.Num() == 1))
+			{
+				WeakSelectedPages.Add(SelectedPages[0]);
+			}
+			DetailsView->SetObjects(WeakSelectedPages);
+		}
+	}
+}
+
+
+#undef LOCTEXT_NAMESPACE
diff --git a/Engine/Plugins/Experimental/RenderPages/Source/RenderPagesEditor/Private/UI/SRenderPagesPage.h b/Engine/Plugins/Experimental/RenderPages/Source/RenderPagesEditor/Private/UI/SRenderPagesPage.h
new file mode 100644
--- /dev/null
+++ b/Engine/Plugins/Experimental/RenderPages/Source/RenderPagesEditor/Private/UI/SRenderPagesPage.h
@@ -0,0 +1,48 @@
+// Copyright Epic Games, Inc. All Rights Reserved.
+
+#pragma once
+
+#include "CoreMinimal.h"
+#include "Widgets/SCompoundWidget.h"
+#include "Misc/NotifyHook.h"
+
+
+class IDetailsView;
+
+namespace UE::RenderPages
+{
+	class IRenderPageCollectionEditor;
+}
+
+
+namespace UE::RenderPages::Private
+{
+	/**
+	 * A widget with which the user can modify the selected render page.
+	 * Can only modify 1 render page at a time, this widget will show nothing when 0 or 2+ render pages are selected.
+	 */
+	class SRenderPagesPage : public SCompoundWidget, public FNotifyHook
+	{
+	public:
+		SLATE_BEGIN_ARGS(SRenderPagesPage) {}
+		SLATE_END_ARGS()
+
+		void Construct(const FArguments& InArgs, TSharedPtr<IRenderPageCollectionEditor> InBlueprintEditor);
+
+		// FNotifyHook interface
+		virtual void NotifyPreChange(FEditPropertyChain* PropertyAboutToChange) override;
+		virtual void NotifyPostChange(const FPropertyChangedEvent& PropertyChangedEvent, FEditPropertyChain* PropertyThatChanged) override;
+		// End of FNotifyHook interface
+
+	private:
+		/** Callback for updating the details view. */
+		void HandlePagesSelectionChanged();
+
+	private:
+		/** A reference to the BP Editor that owns this collection. */
+		TWeakPtr<IRenderPageCollectionEditor> BlueprintEditorWeakPtr;
+
+		/** A reference to the details view. */
+		TSharedPtr<IDetailsView> DetailsView;
+	};
+}
