First usable version of the render pages plugin.

diff --git a/Engine/Plugins/Experimental/RenderPages/Source/RenderPagesEditor/Private/UI/SRenderPagesPageList.cpp b/Engine/Plugins/Experimental/RenderPages/Source/RenderPagesEditor/Private/UI/SRenderPagesPageList.cpp
new file mode 100644
--- /dev/null
+++ b/Engine/Plugins/Experimental/RenderPages/Source/RenderPagesEditor/Private/UI/SRenderPagesPageList.cpp
@@ -0,0 +1,575 @@
+// Copyright Epic Games, Inc. All Rights Reserved.
+
+#include "UI/SRenderPagesPageList.h"
+
+#include "UI/Components/SRenderPagesDragHandle.h"
+#include "UI/Components/SRenderPagesEditableTextBlock.h"
+#include "UI/Components/SRenderPagesFileSelectorTextBlock.h"
+#include "RenderPage/RenderPageCollection.h"
+#include "RenderPage/RenderPageMoviePipelineJob.h"
+#include "RenderPage/RenderPageManager.h"
+#include "IRenderPageCollectionEditor.h"
+#include "IRenderPagesModule.h"
+
+#include "DesktopPlatformModule.h"
+#include "Misc/MessageDialog.h"
+#include "MoviePipelineExecutor.h"
+#include "MoviePipelineQueue.h"
+#include "PropertyCustomizationHelpers.h"
+#include "SlateOptMacros.h"
+#include "Styling/AppStyle.h"
+#include "Widgets/Input/SCheckBox.h"
+#include "Widgets/Input/SSearchBox.h"
+#include "Widgets/Layout/SScaleBox.h"
+
+#define LOCTEXT_NAMESPACE "SRenderPagesPageList"
+
+
+namespace UE::RenderPages::Private::FRenderPagesCollectionColumns
+{
+	const FName DragDropHandle = TEXT("DragDropHandle");
+	const FName IsEnabled = TEXT("IsEnabled");
+	const FName PageId = TEXT("PageId");
+	const FName PageName = TEXT("PageName");
+	const FName OutputDirectory = TEXT("OutputDirectory");
+	const FName RenderPreset = TEXT("RenderPreset");
+	const FName StartFrame = TEXT("StartFrame");
+	const FName EndFrame = TEXT("EndFrame");
+	const FName Tags = TEXT("Tags");
+	const FName Duration = TEXT("Duration");
+	const FName RenderingStatus = TEXT("Status");
+}
+
+
+void UE::RenderPages::Private::SRenderPagesPageList::Tick(const FGeometry&, const double, const float)
+{
+	if (const TSharedPtr<IRenderPageCollectionEditor> BlueprintEditor = BlueprintEditorWeakPtr.Pin())
+	{
+		if (RenderPagesCollectionWeakPtr != BlueprintEditor->GetInstance())
+		{
+			Refresh();
+		}
+	}
+}
+
+BEGIN_SLATE_FUNCTION_BUILD_OPTIMIZATION
+
+void UE::RenderPages::Private::SRenderPagesPageList::Construct(const FArguments& InArgs, TSharedPtr<IRenderPageCollectionEditor> InBlueprintEditor)
+{
+	BlueprintEditorWeakPtr = InBlueprintEditor;
+
+	Refresh();
+	InBlueprintEditor->OnRenderPageCreated().AddSP(this, &SRenderPagesPageList::OnRenderPageCreated);
+	InBlueprintEditor->OnRenderPagesChanged().AddSP(this, &SRenderPagesPageList::Refresh);
+	InBlueprintEditor->OnRenderPagesBatchRenderingStarted().AddSP(this, &SRenderPagesPageList::OnBatchRenderingStarted);
+	InBlueprintEditor->OnRenderPagesBatchRenderingFinished().AddSP(this, &SRenderPagesPageList::OnBatchRenderingFinished);
+
+	ChildSlot
+	[
+		SNew(SVerticalBox)
+
+		+ SVerticalBox::Slot()
+		.AutoHeight()
+		.Padding(2.f)
+		[
+			SNew(SHorizontalBox)
+			// Search Box
+			+ SHorizontalBox::Slot()
+			.Padding(4.f, 2.f)
+			[
+				SAssignNew(RenderPagesSearchBox, SSearchBox)
+				.HintText(LOCTEXT("Search_HintText", "Search Tags | Text"))
+				.OnTextChanged(this, &SRenderPagesPageList::OnSearchBarTextChanged)
+			]
+			// Filters
+			+ SHorizontalBox::Slot()
+			.AutoWidth()
+			.Padding(0.f, 2.f, 2.f, 2.f)
+			[
+				SNew(SComboButton)
+				.ComboButtonStyle(FAppStyle::Get(), "GenericFilters.ComboButtonStyle")
+				.ForegroundColor(FLinearColor::White)
+				.ToolTipText(LOCTEXT("Filters_Tooltip", "Filter options for the Pages Collection."))
+				.HasDownArrow(true)
+				.ContentPadding(0.f)
+				.ButtonContent()
+				[
+					SNew(STextBlock)
+					.TextStyle(FAppStyle::Get(), "GenericFilters.TextStyle")
+					.Font(FAppStyle::Get().GetFontStyle("FontAwesome.12"))
+					.Text(FText::FromString(FString(TEXT("\xf0b0"))) /*fa-filter*/)
+				]
+			]
+		]
+		// Pages Collection
+		+ SVerticalBox::Slot()
+		.FillHeight(1.f)
+		[
+			SNew(SBorder)
+			.BorderImage(FAppStyle::GetBrush("ToolPanel.GroupBorder"))
+			.Padding(0.f)
+			[
+				SAssignNew(RenderPageListWidget, SListView<URenderPage*>)
+				.ItemHeight(20.0f)
+				.OnGenerateRow(this, &SRenderPagesPageList::HandlePagesCollectionGenerateRow)
+				.OnSelectionChanged(this, &SRenderPagesPageList::HandlePagesCollectionSelectionChanged)
+				.SelectionMode(ESelectionMode::Multi)
+				.ClearSelectionOnClick(false)
+				.ListItemsSource(&RenderPages)
+				.HeaderRow(
+					SNew(SHeaderRow)
+
+					+ SHeaderRow::Column(FRenderPagesCollectionColumns::DragDropHandle)
+					.DefaultLabel(LOCTEXT("PagesCollectionDragDropHandleColumnHeader", ""))
+					.FixedWidth(36.0f)
+
+					+ SHeaderRow::Column(FRenderPagesCollectionColumns::IsEnabled)
+					.DefaultLabel(LOCTEXT("PagesCollectionIDColumnHeader", "Enabled"))
+					.FixedWidth(30.0f) //55.0f for text : "Enabled"
+					[
+						SAssignNew(RenderPageEnabledHeaderCheckbox, SCheckBox)
+						.IsChecked(true)
+						.OnCheckStateChanged(this, &SRenderPagesPageList::OnHeaderCheckboxToggled)
+					]
+
+					+ SHeaderRow::Column(FRenderPagesCollectionColumns::PageId)
+					.DefaultLabel(LOCTEXT("PagesCollectionIDColumnHeader", "Page ID"))
+					.FillWidth(0.3f)
+
+					+ SHeaderRow::Column(FRenderPagesCollectionColumns::PageName)
+					.DefaultLabel(LOCTEXT("PagesCollectionNameColumnHeader", "Page Name"))
+					.FillWidth(0.3f)
+
+					+ SHeaderRow::Column(FRenderPagesCollectionColumns::OutputDirectory)
+					.DefaultLabel(LOCTEXT("PagesCollectionOutDirColumnHeader", "Output Directory"))
+					.FillWidth(0.7f)
+
+					+ SHeaderRow::Column(FRenderPagesCollectionColumns::RenderPreset)
+					.DefaultLabel(LOCTEXT("PagesCollectionRenderPresetColumnHeader", "Render Preset"))
+					.FillWidth(0.5f)
+
+					+ SHeaderRow::Column(FRenderPagesCollectionColumns::StartFrame)
+					.DefaultLabel(LOCTEXT("PagesCollectionStartFrameColumnHeader", "Start Frame"))
+					.FixedWidth(80.0f)
+
+					+ SHeaderRow::Column(FRenderPagesCollectionColumns::EndFrame)
+					.DefaultLabel(LOCTEXT("PagesCollectionEndFrameColumnHeader", "End Frame"))
+					.FixedWidth(80.0f)
+
+					+ SHeaderRow::Column(FRenderPagesCollectionColumns::Tags)
+					.DefaultLabel(LOCTEXT("PagesCollectionTagsColumnHeader", "Tags"))
+					.FillWidth(0.7f)
+
+					+ SHeaderRow::Column(FRenderPagesCollectionColumns::Duration)
+					.DefaultLabel(LOCTEXT("PagesCollectionEstDurColumnHeader", "Est Duration"))
+					.FixedWidth(120.0f)
+				)
+			]
+		]
+	];
+
+	Refresh();
+}
+
+END_SLATE_FUNCTION_BUILD_OPTIMIZATION
+
+void UE::RenderPages::Private::SRenderPagesPageList::OnRenderPageCreated(URenderPage* Page)
+{
+	if (!RenderPageEnabledHeaderCheckbox.IsValid())
+	{
+		return;
+	}
+	Page->SetIsEnabled(RenderPageEnabledHeaderCheckbox->GetCheckedState() != ECheckBoxState::Unchecked);
+}
+
+void UE::RenderPages::Private::SRenderPagesPageList::OnHeaderCheckboxToggled(ECheckBoxState State)
+{
+	if (const TSharedPtr<IRenderPageCollectionEditor> BlueprintEditor = BlueprintEditorWeakPtr.Pin())
+	{
+		if (URenderPageCollection* Collection = BlueprintEditor->GetInstance(); IsValid(Collection))
+		{
+			bool bRefresh = false;
+
+			for (URenderPage* Page : Collection->GetRenderPages())
+			{
+				Page->SetIsEnabled(State == ECheckBoxState::Checked);
+				bRefresh = true;
+			}
+
+			if (bRefresh)
+			{
+				Refresh();
+			}
+		}
+	}
+}
+
+ECheckBoxState UE::RenderPages::Private::SRenderPagesPageList::GetDesiredHeaderEnabledCheckboxState()
+{
+	ECheckBoxState State = ECheckBoxState::Checked;
+	if (const TSharedPtr<IRenderPageCollectionEditor> BlueprintEditor = BlueprintEditorWeakPtr.Pin())
+	{
+		if (URenderPageCollection* Collection = BlueprintEditor->GetInstance(); IsValid(Collection))
+		{
+			bool bFirstPage = true;
+			for (URenderPage* Page : Collection->GetRenderPages())
+			{
+				ECheckBoxState PageState = (Page->GetIsEnabled() ? ECheckBoxState::Checked : ECheckBoxState::Unchecked);
+				if (bFirstPage)
+				{
+					State = PageState;
+					bFirstPage = false;
+				}
+				else if (State != PageState)
+				{
+					return ECheckBoxState::Undetermined;
+				}
+			}
+		}
+	}
+	return State;
+}
+
+void UE::RenderPages::Private::SRenderPagesPageList::AddRenderStatusColumn()
+{
+	if (!RenderPageListWidget.IsValid())
+	{
+		return;
+	}
+	RenderPageListWidget->GetHeaderRow()->AddColumn(SHeaderRow::Column(FRenderPagesCollectionColumns::RenderingStatus)
+		.DefaultLabel(LOCTEXT("PagesCollectionRenderStatusColumnHeader", "Render Status"))
+		.FillWidth(0.5f));
+}
+
+void UE::RenderPages::Private::SRenderPagesPageList::RemoveRenderStatusColumn()
+{
+	if (!RenderPageListWidget.IsValid())
+	{
+		return;
+	}
+	RenderPageListWidget->GetHeaderRow()->RemoveColumn(FRenderPagesCollectionColumns::RenderingStatus);
+}
+
+
+void UE::RenderPages::Private::SRenderPagesPageList::Refresh()
+{
+	if (const TSharedPtr<IRenderPageCollectionEditor> BlueprintEditor = BlueprintEditorWeakPtr.Pin())
+	{
+		const bool bIsBatchRendering = BlueprintEditor->IsBatchRendering();// show all pages during a batch render, ignore the search bar
+		const FString SearchBarContent = (!RenderPagesSearchBox.IsValid() ? TEXT("") : RenderPagesSearchBox->GetText().ToString());
+
+		RenderPages.Empty();
+		RenderPagesCollectionWeakPtr = BlueprintEditor->GetInstance();
+		if (URenderPageCollection* Collection = RenderPagesCollectionWeakPtr.Get(); IsValid(Collection))
+		{
+			for (URenderPage* Page : Collection->GetRenderPages())
+			{
+				if (bIsBatchRendering || Page->MatchesSearchTerm(SearchBarContent))
+				{
+					RenderPages.Add(Page);
+				}
+			}
+		}
+
+		RefreshHeaderEnabledCheckbox();
+
+		RemoveRenderStatusColumn();
+		if (bIsBatchRendering)
+		{
+			AddRenderStatusColumn();
+		}
+
+		if (RenderPageListWidget.IsValid())
+		{
+			RenderPageListWidget->RebuildList(); // rebuild is needed (instead of using RequestListRefresh()), because otherwise it won't show the changes made to the URenderPage variables
+		}
+	}
+}
+
+void UE::RenderPages::Private::SRenderPagesPageList::RefreshHeaderEnabledCheckbox()
+{
+	if (!RenderPageEnabledHeaderCheckbox.IsValid())
+	{
+		return;
+	}
+	RenderPageEnabledHeaderCheckbox->SetIsChecked(GetDesiredHeaderEnabledCheckboxState());
+}
+
+TSharedRef<ITableRow> UE::RenderPages::Private::SRenderPagesPageList::HandlePagesCollectionGenerateRow(URenderPage* Item, const TSharedRef<STableViewBase>& OwnerTable)
+{
+	TSharedPtr<SWidget> BaseThis = AsShared();
+	TSharedPtr<SRenderPagesPageList> This = StaticCastSharedPtr<SRenderPagesPageList>(BaseThis);
+
+	return SNew(SRenderPagesPageListTableRow, OwnerTable, BlueprintEditorWeakPtr, Item, This);
+}
+
+void UE::RenderPages::Private::SRenderPagesPageList::HandlePagesCollectionSelectionChanged(URenderPage* Item, ESelectInfo::Type SelectInfo)
+{
+	if (const TSharedPtr<IRenderPageCollectionEditor> BlueprintEditor = BlueprintEditorWeakPtr.Pin())
+	{
+		BlueprintEditor->SetSelectedRenderPages(RenderPageListWidget->GetSelectedItems());
+	}
+}
+
+
+void UE::RenderPages::Private::SRenderPagesPageListTableRow::Construct(const FArguments& InArgs, const TSharedRef<STableViewBase>& InOwnerTableView, TWeakPtr<IRenderPageCollectionEditor> InBlueprintEditor, URenderPage* InRenderPage, const TSharedPtr<SRenderPagesPageList>& InPageListWidget)
+{
+	BlueprintEditorWeakPtr = InBlueprintEditor;
+	RenderPage = InRenderPage;
+	PageListWidget = InPageListWidget;
+
+	SMultiColumnTableRow<URenderPage*>::Construct(FSuperRowType::FArguments()
+		.OnCanAcceptDrop(this, &SRenderPagesPageListTableRow::OnCanAcceptDrop)
+		.OnAcceptDrop(this, &SRenderPagesPageListTableRow::OnAcceptDrop),
+		InOwnerTableView);
+}
+
+TOptional<EItemDropZone> UE::RenderPages::Private::SRenderPagesPageListTableRow::OnCanAcceptDrop(const FDragDropEvent& Event, EItemDropZone InItemDropZone, URenderPage* Page)
+{
+	if (BlueprintEditorWeakPtr.IsValid() && Event.GetOperationAs<FRenderPagesPageListTableRowDragDropOp>())
+	{
+		if (InItemDropZone == EItemDropZone::OntoItem)
+		{
+			return EItemDropZone::BelowItem;
+		}
+		return InItemDropZone;
+	}
+	return TOptional<EItemDropZone>();
+}
+
+FReply UE::RenderPages::Private::SRenderPagesPageListTableRow::OnAcceptDrop(const FDragDropEvent& Event, EItemDropZone InItemDropZone, URenderPage* Page)
+{
+	if (const TSharedPtr<IRenderPageCollectionEditor> BlueprintEditor = BlueprintEditorWeakPtr.Pin())
+	{
+		if (const TSharedPtr<FRenderPagesPageListTableRowDragDropOp> DragDropOp = Event.GetOperationAs<FRenderPagesPageListTableRowDragDropOp>())
+		{
+			if (IRenderPagesModule::Get().GetManager().DragDropPage(BlueprintEditor->GetInstance(), DragDropOp->GetPage(), Page, (InItemDropZone != EItemDropZone::AboveItem)))
+			{
+				BlueprintEditor->MarkAsModified();
+				BlueprintEditor->OnRenderPagesChanged().Broadcast();
+				return FReply::Handled();
+			}
+		}
+	}
+	return FReply::Unhandled();
+}
+
+TSharedRef<SWidget> UE::RenderPages::Private::SRenderPagesPageListTableRow::GenerateWidgetForColumn(const FName& ColumnName)
+{
+	if (!IsValid(RenderPage))
+	{
+		return SNullWidget::NullWidget;
+	}
+
+	if (ColumnName == FRenderPagesCollectionColumns::DragDropHandle)
+	{
+		return SNew(SBox)
+			.Padding(FMargin(0.0f, 2.0f, 2.0f, 2.0f))
+			[
+				SNew(SScaleBox)
+				.HAlign(HAlign_Center)
+				.VAlign(VAlign_Center)
+				.Stretch(EStretch::ScaleToFit)
+				.StretchDirection(EStretchDirection::Both)
+				[
+					SNew(SRenderPagesDragHandle<FRenderPagesPageListTableRowDragDropOp>, RenderPage)
+					.Widget(SharedThis(this))
+				]
+			];
+	}
+	else if (ColumnName == FRenderPagesCollectionColumns::IsEnabled)
+	{
+		return SNew(SBox)
+			.HAlign(HAlign_Center)
+			[
+				SNew(SCheckBox)
+				.IsChecked(RenderPage->GetIsEnabled())
+				.OnCheckStateChanged_Lambda([this](ECheckBoxState State)
+				{
+					RenderPage->SetIsEnabled(State == ECheckBoxState::Checked);
+					if (PageListWidget.IsValid())
+					{
+						PageListWidget->RefreshHeaderEnabledCheckbox();
+					}
+					if (const TSharedPtr<IRenderPageCollectionEditor> BlueprintEditor = BlueprintEditorWeakPtr.Pin())
+					{
+						BlueprintEditor->MarkAsModified();
+					}
+				})
+			];
+	}
+	else if (ColumnName == FRenderPagesCollectionColumns::PageId)
+	{
+		return SNew(SRenderPagesEditableTextBlock)
+			.Text(FText::FromString(RenderPage->GetPageId()))
+			.OnTextCommitted_Lambda([this](const FText& InLabel, ETextCommit::Type InCommitInfo) -> FText
+			{
+				const FString OldPageId = RenderPage->GetPageId();
+				const FString NewPageId = URenderPage::PurgePageIdOrReturnEmptyString(InLabel.ToString());
+				if (NewPageId.IsEmpty() || (RenderPage->GetPageId() == NewPageId))
+				{
+					return FText::FromString(OldPageId);
+				}
+
+				if (const TSharedPtr<IRenderPageCollectionEditor> BlueprintEditor = BlueprintEditorWeakPtr.Pin())
+				{
+					if (URenderPageCollection* PageCollection = BlueprintEditor->GetInstance(); IsValid(PageCollection))
+					{
+						if (IRenderPagesModule::Get().GetManager().DoesPageIdExist(PageCollection, NewPageId))
+						{
+							const FText TitleText = LOCTEXT("PageIdNotUniqueTitle", "Duplicate Page IDs");
+							FMessageDialog::Open(
+								EAppMsgType::Ok,
+								FText::Format(LOCTEXT("PageIdNotUniqueMessage", "Page ID \"{0}\" is not unique."), FText::FromString(NewPageId)),
+								&TitleText);
+							return FText::FromString(OldPageId);
+						}
+
+						RenderPage->SetPageId(NewPageId);
+						BlueprintEditor->MarkAsModified();
+						return FText::FromString(RenderPage->GetPageId());
+					}
+				}
+				return FText::FromString(OldPageId);
+			});
+	}
+	else if (ColumnName == FRenderPagesCollectionColumns::PageName)
+	{
+		return SNew(SRenderPagesEditableTextBlock)
+			.Text(FText::FromString(RenderPage->GetPageName()))
+			.OnTextCommitted_Lambda([this](const FText& InLabel, ETextCommit::Type InCommitInfo) -> FText
+			{
+				RenderPage->SetPageName(InLabel.ToString());
+				if (const TSharedPtr<IRenderPageCollectionEditor> BlueprintEditor = BlueprintEditorWeakPtr.Pin())
+				{
+					BlueprintEditor->MarkAsModified();
+				}
+				return FText::FromString(RenderPage->GetPageName());
+			});
+	}
+	else if (ColumnName == FRenderPagesCollectionColumns::OutputDirectory)
+	{
+		return SNew(SRenderPagesFileSelectorTextBlock)
+			.Text(FText::FromString(RenderPage->GetOutputDirectoryForDisplay()))
+			.FolderPath_Lambda([this]() -> FString
+			{
+				return RenderPage->GetOutputDirectory();
+			})
+			.OnTextCommitted_Lambda([this](const FText& InLabel, ETextCommit::Type InCommitInfo) -> FText
+			{
+				RenderPage->SetOutputDirectory(InLabel.ToString());
+				if (const TSharedPtr<IRenderPageCollectionEditor> BlueprintEditor = BlueprintEditorWeakPtr.Pin())
+				{
+					BlueprintEditor->MarkAsModified();
+				}
+				return FText::FromString(RenderPage->GetOutputDirectoryForDisplay());
+			});
+	}
+	else if (ColumnName == FRenderPagesCollectionColumns::RenderPreset)
+	{
+		return SNew(SObjectPropertyEntryBox)
+			.AllowedClass(UMoviePipelineMasterConfig::StaticClass())
+			.ObjectPath_Lambda([this]()-> FString
+			{
+				if (UMoviePipelineMasterConfig* Preset = RenderPage->GetRenderPreset(); IsValid(Preset))
+				{
+					return Preset->GetPathName();
+				}
+				return FString();
+			})
+			.OnObjectChanged_Lambda([this](const FAssetData& AssetData) -> void
+			{
+				RenderPage->SetRenderPreset(nullptr);
+				if (UObject* AssetDataAsset = AssetData.GetAsset(); IsValid(AssetDataAsset))
+				{
+					if (UMoviePipelineMasterConfig* Preset = Cast<UMoviePipelineMasterConfig>(AssetDataAsset))
+					{
+						RenderPage->SetRenderPreset(Preset);
+					}
+				}
+				if (const TSharedPtr<IRenderPageCollectionEditor> BlueprintEditor = BlueprintEditorWeakPtr.Pin())
+				{
+					BlueprintEditor->MarkAsModified();
+					BlueprintEditor->OnRenderPagesChanged().Broadcast();
+				}
+			})
+			.AllowClear(true)
+			.DisplayUseSelected(true)
+			.DisplayBrowse(true)
+			.DisplayThumbnail(false);
+	}
+	else if ((ColumnName == FRenderPagesCollectionColumns::StartFrame) || (ColumnName == FRenderPagesCollectionColumns::EndFrame))
+	{
+		FText Text;
+		if (TOptional<int32> Frame = ((ColumnName == FRenderPagesCollectionColumns::StartFrame) ? RenderPage->GetStartFrame() : RenderPage->GetEndFrame()))
+		{
+			Text = FText::AsNumber(*Frame);
+		}
+		return SNew(SBox)
+			.VAlign(VAlign_Center)
+			.HAlign(HAlign_Right)
+			[
+				SNew(STextBlock).Text(Text)
+			];
+	}
+	else if (ColumnName == FRenderPagesCollectionColumns::Tags)
+	{
+		//TODO: add support for tags
+	}
+	else if (ColumnName == FRenderPagesCollectionColumns::Duration)
+	{
+		FText Text;
+		if (TOptional<double> Duration = RenderPage->GetDurationInSeconds())
+		{
+			FTimespan Timespan = FTimespan::FromSeconds(*Duration);
+			int32 Hours = static_cast<int32>(Timespan.GetTotalHours());
+			int32 Minutes = Timespan.GetMinutes();
+			int32 Seconds = Timespan.GetSeconds();
+
+			FNumberFormattingOptions NumberFormattingOptions;
+			NumberFormattingOptions.MinimumIntegralDigits = 2;
+			NumberFormattingOptions.MaximumIntegralDigits = 2;
+
+			FText TimespanFormatPattern = NSLOCTEXT("Timespan", "Format_HoursMinutesSeconds", "{Hours}:{Minutes}:{Seconds}");
+			FFormatNamedArguments TimeArguments;
+			TimeArguments.Add(TEXT("Hours"), Hours);
+			TimeArguments.Add(TEXT("Minutes"), FText::AsNumber(Minutes, &NumberFormattingOptions));
+			TimeArguments.Add(TEXT("Seconds"), FText::AsNumber(Seconds, &NumberFormattingOptions));
+
+			Text = FText::Format(TimespanFormatPattern, TimeArguments);
+		}
+		return SNew(SBox)
+			.VAlign(VAlign_Center)
+			[
+				SNew(STextBlock).Text(Text)
+			];
+	}
+	else if (ColumnName == FRenderPagesCollectionColumns::RenderingStatus)
+	{
+		return SNew(SBox)
+			.VAlign(VAlign_Center)
+			[
+				SNew(STextBlock)
+				.Text_Lambda([this]() -> FText
+				{
+					return GetRenderStatusText();
+				})
+			];
+	}
+	return SNullWidget::NullWidget;
+}
+
+FText UE::RenderPages::Private::SRenderPagesPageListTableRow::GetRenderStatusText() const
+{
+	if (const TSharedPtr<IRenderPageCollectionEditor> BlueprintEditor = BlueprintEditorWeakPtr.Pin())
+	{
+		if (URenderPagesMoviePipelineRenderJob* RenderJob = BlueprintEditor->GetBatchRenderJob(); IsValid(RenderJob))
+		{
+			return FText::FromString(RenderJob->GetPageStatus(RenderPage));
+		}
+	}
+	return FText();
+}
+
+
+#undef LOCTEXT_NAMESPACE
diff --git a/Engine/Plugins/Experimental/RenderPages/Source/RenderPagesEditor/Private/UI/SRenderPagesPageList.h b/Engine/Plugins/Experimental/RenderPages/Source/RenderPagesEditor/Private/UI/SRenderPagesPageList.h
new file mode 100644
--- /dev/null
+++ b/Engine/Plugins/Experimental/RenderPages/Source/RenderPagesEditor/Private/UI/SRenderPagesPageList.h
@@ -0,0 +1,165 @@
+// Copyright Epic Games, Inc. All Rights Reserved.
+
+#pragma once
+
+#include "CoreMinimal.h"
+#include "DragAndDrop/DecoratedDragDropOp.h"
+#include "RenderPage/RenderPageCollection.h"
+#include "Styling/AppStyle.h"
+#include "Widgets/SCompoundWidget.h"
+#include "Widgets/Text/SInlineEditableTextBlock.h"
+#include "Widgets/Views/STableRow.h"
+#include "Widgets/Views/STableViewBase.h"
+
+
+class SCheckBox;
+class SSearchBox;
+class URenderPage;
+class URenderPagesMoviePipelineRenderJob;
+
+namespace UE::RenderPages
+{
+	class IRenderPageCollectionEditor;
+}
+
+
+namespace UE::RenderPages::Private
+{
+	DECLARE_DELEGATE_RetVal_TwoParams(FText, FOnRenderPagesPageListEditableTextBlockTextCommitted, const FText&, ETextCommit::Type);
+
+
+	/**
+	 * A widget with which the user can see and modify the list of pages the render page collection contains.
+	 */
+	class SRenderPagesPageList : public SCompoundWidget
+	{
+	public:
+		SLATE_BEGIN_ARGS(SRenderPagesPageList) {}
+		SLATE_END_ARGS()
+
+		virtual void Tick(const FGeometry&, const double, const float) override;
+		void Construct(const FArguments& InArgs, TSharedPtr<IRenderPageCollectionEditor> InBlueprintEditor);
+
+		/** Refreshes the content of this widget. */
+		void Refresh();
+
+		/** Refreshes the state of the header is-page-enabled checkbox. */
+		void RefreshHeaderEnabledCheckbox();
+
+	private:
+		/** Gets called when a page is created. */
+		void OnRenderPageCreated(URenderPage* Page);
+
+		/** Gets called when the header is-page-enabled checkbox is toggled. */
+		void OnHeaderCheckboxToggled(ECheckBoxState State);
+
+		/** Gets the desired state of the header is-page-enabled checkbox. */
+		ECheckBoxState GetDesiredHeaderEnabledCheckboxState();
+
+		/** Adds the render status column to the page list. */
+		void AddRenderStatusColumn();
+
+		/** Removes the render status column from the page list. */
+		void RemoveRenderStatusColumn();
+
+	private:
+		void OnBatchRenderingStarted(URenderPagesMoviePipelineRenderJob* RenderJob) { Refresh(); }
+		void OnBatchRenderingFinished(URenderPagesMoviePipelineRenderJob* RenderJob) { Refresh(); }
+		void OnSearchBarTextChanged(const FText& Text) { Refresh(); }
+
+	private:
+		/** Callback for generating a row widget in the session tree view. */
+		TSharedRef<ITableRow> HandlePagesCollectionGenerateRow(URenderPage* Item, const TSharedRef<STableViewBase>& OwnerTable);
+
+		/** Callback for session tree view selection changes. */
+		void HandlePagesCollectionSelectionChanged(URenderPage* Item, ESelectInfo::Type SelectInfo);
+
+	private:
+		/** A reference to the BP Editor that owns this collection. */
+		TWeakPtr<IRenderPageCollectionEditor> BlueprintEditorWeakPtr;
+
+		/** The render page collection that is shown in the UI. */
+		TWeakObjectPtr<URenderPageCollection> RenderPagesCollectionWeakPtr;
+
+		/** The render pages that are shown in the UI. */
+		TArray<URenderPage*> RenderPages;
+
+		/** The widget that lists the render pages. */
+		TSharedPtr<SListView<URenderPage*>> RenderPageListWidget;
+
+		/** The search bar widget. */
+		TSharedPtr<SSearchBox> RenderPagesSearchBox;
+
+		/** The header checkbox for the enable/disable column. */
+		TSharedPtr<SCheckBox> RenderPageEnabledHeaderCheckbox;
+	};
+
+
+	/**
+	 * The widget that represents a single render page (a single row).
+	 */
+	class SRenderPagesPageListTableRow : public SMultiColumnTableRow<URenderPage*>
+	{
+	public:
+		SLATE_BEGIN_ARGS(SRenderPagesPageListTableRow) {}
+		SLATE_END_ARGS()
+
+		void Construct(const FArguments& InArgs, const TSharedRef<STableViewBase>& InOwnerTableView, TWeakPtr<IRenderPageCollectionEditor> InBlueprintEditor, URenderPage* InRenderPage, const TSharedPtr<SRenderPagesPageList>& InPageListWidget);
+		TOptional<EItemDropZone> OnCanAcceptDrop(const FDragDropEvent& Event, EItemDropZone InItemDropZone, URenderPage* Page);
+		FReply OnAcceptDrop(const FDragDropEvent& Event, EItemDropZone InItemDropZone, URenderPage* Page);
+		virtual TSharedRef<SWidget> GenerateWidgetForColumn(const FName& ColumnName) override;
+		FText GetRenderStatusText() const;
+
+	private:
+		/** A reference to the BP Editor that owns this collection. */
+		TWeakPtr<IRenderPageCollectionEditor> BlueprintEditorWeakPtr;
+
+		/** A reference to the render page model. */
+		TObjectPtr<URenderPage> RenderPage;
+
+		/** A reference to the page list (the parent widget). */
+		TSharedPtr<SRenderPagesPageList> PageListWidget;
+	};
+
+
+	/**
+	 * The class that makes it possible to drag and drop render pages (allowing the user to reorganize the render pages list).
+	 */
+	class FRenderPagesPageListTableRowDragDropOp final : public FDecoratedDragDropOp
+	{
+	public:
+		DRAG_DROP_OPERATOR_TYPE(FRenderPagesPageListTableRowDragDropOp, FDragDropOperation)
+
+		using WidgetType = SRenderPagesPageListTableRow;
+		using HeldItemType = TObjectPtr<URenderPage>;
+
+		FRenderPagesPageListTableRowDragDropOp(const TSharedPtr<WidgetType> InWidget, const HeldItemType InPage)
+			: Page(InPage)
+		{
+			DecoratorWidget = SNew(SBorder)
+				.Padding(0.f)
+				.BorderImage(FAppStyle::Get().GetBrush("Graph.ConnectorFeedback.Border"))
+				.Content()
+				[
+					InWidget.ToSharedRef()
+				];
+		}
+
+		HeldItemType GetPage() const
+		{
+			return Page;
+		}
+
+		virtual TSharedPtr<SWidget> GetDefaultDecorator() const override
+		{
+			return DecoratorWidget;
+		}
+
+	private:
+		/** The held item. */
+		HeldItemType Page;
+
+		/** Holds the displayed widget. */
+		TSharedPtr<SWidget> DecoratorWidget;
+	};
+}
diff --git a/Engine/Plugins/Experimental/RenderPages/Source/RenderPagesEditor/Private/UI/SRenderPagesPageViewer.cpp b/Engine/Plugins/Experimental/RenderPages/Source/RenderPagesEditor/Private/UI/SRenderPagesPageViewer.cpp
new file mode 100644
--- /dev/null
+++ b/Engine/Plugins/Experimental/RenderPages/Source/RenderPagesEditor/Private/UI/SRenderPagesPageViewer.cpp
@@ -0,0 +1,116 @@
+// Copyright Epic Games, Inc. All Rights Reserved.
+
+#include "UI/SRenderPagesPageViewer.h"
+#include "UI/SRenderPagesPageViewerLive.h"
+#include "UI/SRenderPagesPageViewerPreview.h"
+#include "UI/SRenderPagesPageViewerRendered.h"
+#include "IRenderPageCollectionEditor.h"
+#include "SlateOptMacros.h"
+#include "Styles/RenderPagesEditorStyle.h"
+#include "Widgets/Input/SButton.h"
+
+#define LOCTEXT_NAMESPACE "SRenderPagesPageViewer"
+
+
+BEGIN_SLATE_FUNCTION_BUILD_OPTIMIZATION
+
+void UE::RenderPages::Private::SRenderPagesPageViewer::Construct(const FArguments& InArgs, TSharedPtr<IRenderPageCollectionEditor> InBlueprintEditor)
+{
+	BlueprintEditorWeakPtr = InBlueprintEditor;
+	ViewerMode = ERenderPagesPageViewerMode::Live;
+	CachedViewerMode = ERenderPagesPageViewerMode::None;
+
+	SAssignNew(WidgetContainer, SBorder)
+		.Padding(0.0f)
+		.BorderImage(new FSlateNoResource());
+
+	Refresh();
+
+	ChildSlot
+	[
+		SNew(SVerticalBox)
+
+		+ SVerticalBox::Slot()
+		.AutoHeight()
+		[
+			SNew(SBox)
+			.HeightOverride(27.5f)
+			[
+				SNew(SHorizontalBox)
+				+ SHorizontalBox::Slot()
+				[
+					CreateViewerModeButton(FText::FromString("Live"), ERenderPagesPageViewerMode::Live)
+				]
+				+ SHorizontalBox::Slot()
+				[
+					CreateViewerModeButton(FText::FromString("Preview"), ERenderPagesPageViewerMode::Preview)
+				]
+				+ SHorizontalBox::Slot()
+				[
+					CreateViewerModeButton(FText::FromString("Rendered"), ERenderPagesPageViewerMode::Rendered)
+				]
+				+ SHorizontalBox::Slot()
+				[
+					CreateViewerModeButton(FText::FromString("None"), ERenderPagesPageViewerMode::None)
+				]
+			]
+		]
+
+		+ SVerticalBox::Slot()
+		.FillHeight(1.0f)
+		[
+			WidgetContainer.ToSharedRef()
+		]
+	];
+}
+
+TSharedRef<SWidget> UE::RenderPages::Private::SRenderPagesPageViewer::CreateViewerModeButton(const FText& ButtonText, const ERenderPagesPageViewerMode ButtonViewerMode)
+{
+	return SNew(SButton)
+		.ButtonStyle(FRenderPagesEditorStyle::Get(), "TabButton")
+		.Text(ButtonText)
+		.HAlign(HAlign_Center)
+		.VAlign(VAlign_Center)
+		.OnClicked_Lambda([this, ButtonViewerMode]() -> FReply
+		{
+			ViewerMode = ButtonViewerMode;
+			Refresh();
+			return FReply::Handled();
+		})
+		.IsEnabled_Lambda([this, ButtonViewerMode]() -> bool
+		{
+			return (ViewerMode != ButtonViewerMode);
+		});
+}
+
+END_SLATE_FUNCTION_BUILD_OPTIMIZATION
+
+void UE::RenderPages::Private::SRenderPagesPageViewer::Refresh()
+{
+	if (!WidgetContainer.IsValid() || (ViewerMode == CachedViewerMode))
+	{
+		return;
+	}
+
+	CachedViewerMode = ViewerMode;
+	WidgetContainer->ClearContent();
+
+	if (const TSharedPtr<IRenderPageCollectionEditor> BlueprintEditor = BlueprintEditorWeakPtr.Pin())
+	{
+		if (ViewerMode == ERenderPagesPageViewerMode::Live)
+		{
+			WidgetContainer->SetContent(SNew(SRenderPagesPageViewerLive, BlueprintEditor));
+		}
+		else if (ViewerMode == ERenderPagesPageViewerMode::Preview)
+		{
+			WidgetContainer->SetContent(SNew(SRenderPagesPageViewerPreview, BlueprintEditor));
+		}
+		else if (ViewerMode == ERenderPagesPageViewerMode::Rendered)
+		{
+			WidgetContainer->SetContent(SNew(SRenderPagesPageViewerRendered, BlueprintEditor));
+		}
+	}
+}
+
+
+#undef LOCTEXT_NAMESPACE
diff --git a/Engine/Plugins/Experimental/RenderPages/Source/RenderPagesEditor/Private/UI/SRenderPagesPageViewer.h b/Engine/Plugins/Experimental/RenderPages/Source/RenderPagesEditor/Private/UI/SRenderPagesPageViewer.h
new file mode 100644
--- /dev/null
+++ b/Engine/Plugins/Experimental/RenderPages/Source/RenderPagesEditor/Private/UI/SRenderPagesPageViewer.h
@@ -0,0 +1,61 @@
+// Copyright Epic Games, Inc. All Rights Reserved.
+
+#pragma once
+
+#include "CoreMinimal.h"
+#include "Widgets/SCompoundWidget.h"
+
+
+class SBorder;
+
+namespace UE::RenderPages
+{
+	class IRenderPageCollectionEditor;
+}
+
+
+namespace UE::RenderPages::Private
+{
+	/**
+	 * An enum containing the different page viewer modes that are currently available in the render pages plugin.
+	 */
+	enum class ERenderPagesPageViewerMode : uint8
+	{
+		Live,
+		Preview,
+		Rendered,
+		None
+	};
+
+	/**
+	 * The render page viewer, allows the user to see the expected render output, directly in the editor.
+	 */
+	class SRenderPagesPageViewer : public SCompoundWidget
+	{
+	public:
+		SLATE_BEGIN_ARGS(SRenderPagesPageViewer) {}
+		SLATE_END_ARGS()
+
+		void Construct(const FArguments& InArgs, TSharedPtr<IRenderPageCollectionEditor> InBlueprintEditor);
+
+		/** Refreshes the content of this widget. */
+		void Refresh();
+
+	private:
+		/** Creates a tab button for a viewer mode. */
+		TSharedRef<SWidget> CreateViewerModeButton(const FText& ButtonText, const ERenderPagesPageViewerMode ButtonViewerMode);
+
+	private:
+		/** A reference to the BP Editor that owns this collection. */
+		TWeakPtr<IRenderPageCollectionEditor> BlueprintEditorWeakPtr;
+
+		/** The widget that lists the viewers. */
+		TSharedPtr<SBorder> WidgetContainer;
+
+		/** The current viewer mode that should be shown in the UI. */
+		ERenderPagesPageViewerMode ViewerMode;
+
+		/** The last viewer mode that was shown in the UI. */
+		ERenderPagesPageViewerMode CachedViewerMode;
+	};
+}
diff --git a/Engine/Plugins/Experimental/RenderPages/Source/RenderPagesEditor/Private/UI/SRenderPagesPageViewerLive.cpp b/Engine/Plugins/Experimental/RenderPages/Source/RenderPagesEditor/Private/UI/SRenderPagesPageViewerLive.cpp
new file mode 100644
--- /dev/null
+++ b/Engine/Plugins/Experimental/RenderPages/Source/RenderPagesEditor/Private/UI/SRenderPagesPageViewerLive.cpp
@@ -0,0 +1,289 @@
+// Copyright Epic Games, Inc. All Rights Reserved.
+
+#include "UI/SRenderPagesPageViewerLive.h"
+#include "Camera/CameraComponent.h"
+#include "EngineUtils.h"
+#include "GameFramework/PlayerController.h"
+#include "GameFramework/PlayerStart.h"
+#include "IRenderPageCollectionEditor.h"
+#include "LevelSequence.h"
+#include "LevelSequenceActor.h"
+#include "RenderPage/RenderPageCollection.h"
+#include "SlateOptMacros.h"
+#include "UI/Components/SRenderPagesPageViewerFrameSlider.h"
+#include "Widgets/Layout/SScaleBox.h"
+
+#define LOCTEXT_NAMESPACE "SRenderPagesPageViewerLive"
+
+
+UE::RenderPages::Private::FRenderPagesEditorViewportClient::FRenderPagesEditorViewportClient(FPreviewScene* PreviewScene, const TWeakPtr<SEditorViewport>& InEditorViewportWidget)
+	: FEditorViewportClient(nullptr, PreviewScene, InEditorViewportWidget)
+{
+	FOVAngle = 90;
+	ViewFOV = 90;
+	Invalidate();
+
+	bDisableInput = true;
+	SetGameView(true);
+	SetRealtime(false);// we manually render every frame, because automatic rendering stops temporarily when you're dragging another widget (with the mouse)
+}
+
+
+void UE::RenderPages::Private::SRenderPagesEditorViewport::Tick(const FGeometry&, const double, const float)
+{
+	if (!ViewportClient.IsValid())
+	{
+		return;
+	}
+	ViewportClient->Invalidate();// render
+
+	if (LevelSequenceActor.IsValid())
+	{
+		if (ULevelSequencePlayer* SequencePlayer = LevelSequenceActor->GetSequencePlayer(); IsValid(SequencePlayer))
+		{
+			SequencePlayer->SetPlaybackPosition(FMovieSceneSequencePlaybackParams(LevelSequenceTime, EUpdatePositionMethod::Play));// execute every tick, in case any sequencer values get overwritten (by remote control props for example)
+			if (UCameraComponent* Camera = SequencePlayer->GetActiveCameraComponent(); IsValid(Camera))
+			{
+				ViewportClient->SetViewLocation(Camera->GetComponentLocation());
+				ViewportClient->SetViewRotation(Camera->GetComponentRotation());
+				return;
+			}
+		}
+	}
+
+	if (UWorld* World = GetWorld(); IsValid(World))
+	{
+		if (APlayerController* LocalPlayerController = World->GetFirstPlayerController(); IsValid(LocalPlayerController))
+		{
+			FVector ViewLocation;
+			FRotator ViewRotation;
+			LocalPlayerController->GetPlayerViewPoint(ViewLocation, ViewRotation);
+			ViewportClient->SetViewLocation(ViewLocation);
+			ViewportClient->SetViewRotation(ViewRotation);
+			return;
+		}
+
+		for (TActorIterator<APlayerStart> It(World); It; ++It)
+		{
+			if (APlayerStart* PlayerStart = *It; IsValid(PlayerStart))
+			{
+				ViewportClient->SetViewLocation(PlayerStart->GetActorLocation());
+				ViewportClient->SetViewRotation(PlayerStart->GetActorRotation());
+				return;
+			}
+		}
+	}
+}
+
+void UE::RenderPages::Private::SRenderPagesEditorViewport::Construct(const FArguments& InArgs)
+{
+	LevelSequenceTime = 0.0f;
+	ViewportClient = MakeShareable(new FRenderPagesEditorViewportClient(nullptr, SharedThis(this)));
+
+	SEditorViewport::Construct(SEditorViewport::FArguments());
+
+	if (UWorld* World = GetWorld(); IsValid(World))
+	{
+		FActorSpawnParameters LevelSequenceSpawnParams;
+		LevelSequenceSpawnParams.ObjectFlags = RF_Transient | RF_DuplicateTransient;
+		LevelSequenceActor = TStrongObjectPtr(World->SpawnActor<ALevelSequenceActor>(LevelSequenceSpawnParams));
+		LevelSequenceActor->PlaybackSettings.LoopCount.Value = 0;
+		LevelSequenceActor->PlaybackSettings.bAutoPlay = false;
+		LevelSequenceActor->PlaybackSettings.bPauseAtEnd = true;
+		LevelSequenceActor->PlaybackSettings.bRestoreState = true;
+	}
+}
+
+UE::RenderPages::Private::SRenderPagesEditorViewport::~SRenderPagesEditorViewport()
+{
+	if (LevelSequenceActor.IsValid())
+	{
+		if (ULevelSequencePlayer* SequencePlayer = LevelSequenceActor->GetSequencePlayer(); IsValid(SequencePlayer))
+		{
+			SequencePlayer->Stop();
+		}
+		LevelSequenceActor->Destroy();
+	}
+	LevelSequenceActor.Reset();
+	ViewportClient.Reset();
+}
+
+bool UE::RenderPages::Private::SRenderPagesEditorViewport::ShowSequenceFrame(ULevelSequence* InSequence, const float InTime)
+{
+	LevelSequenceTime = InTime;
+	if (!LevelSequenceActor.IsValid() || !IsValid(InSequence))
+	{
+		return false;
+	}
+	if (LevelSequenceActor->GetSequence() != InSequence)
+	{
+		if (ULevelSequencePlayer* SequencePlayer = LevelSequenceActor->GetSequencePlayer(); IsValid(SequencePlayer))
+		{
+			SequencePlayer->Stop();
+		}
+		LevelSequenceActor->SetSequence(InSequence);
+	}
+	if (ULevelSequencePlayer* SequencePlayer = LevelSequenceActor->GetSequencePlayer(); IsValid(SequencePlayer))
+	{
+		SequencePlayer->SetPlaybackPosition(FMovieSceneSequencePlaybackParams(LevelSequenceTime, EUpdatePositionMethod::Play));
+	}
+	return true;
+}
+
+
+BEGIN_SLATE_FUNCTION_BUILD_OPTIMIZATION
+
+void UE::RenderPages::Private::SRenderPagesPageViewerLive::Construct(const FArguments& InArgs, TSharedPtr<IRenderPageCollectionEditor> InBlueprintEditor)
+{
+	BlueprintEditorWeakPtr = InBlueprintEditor;
+	SelectedPageWeakPtr = nullptr;
+
+	SAssignNew(ViewportWidget, SRenderPagesEditorViewport)
+		.Visibility(EVisibility::Hidden);
+
+	SAssignNew(FrameSlider, SRenderPagesPageViewerFrameSlider)
+		.Visibility(EVisibility::Hidden)
+		.OnValueChanged(this, &SRenderPagesPageViewerLive::FrameSliderValueChanged);
+
+	SelectedPageChanged();
+
+	InBlueprintEditor->OnRenderPagesChanged().AddSP(this, &SRenderPagesPageViewerLive::PagesDataChanged);
+	InBlueprintEditor->OnRenderPagesSelectionChanged().AddSP(this, &SRenderPagesPageViewerLive::SelectedPageChanged);
+	FCoreUObjectDelegates::OnObjectModified.AddSP(this, &SRenderPagesPageViewerLive::OnObjectModified);
+
+	ChildSlot
+	[
+		SNew(SVerticalBox)
+
+		// viewport
+		+ SVerticalBox::Slot()
+		.FillHeight(1.0f)
+		.Padding(0.0f)
+		[
+			SNew(SScaleBox)
+			.HAlign(HAlign_Center)
+			.VAlign(VAlign_Center)
+			.Stretch(EStretch::ScaleToFit)
+			.StretchDirection(EStretchDirection::Both)
+			[
+				SNew(SBox)
+				.HAlign(HAlign_Fill)
+				.VAlign(VAlign_Fill)
+				.WidthOverride(1.0f)
+				.HeightOverride_Lambda([this]() -> float { return 1.0 / (SelectedPageWeakPtr.IsValid() ? SelectedPageWeakPtr->GetOutputAspectRatio() : 1.0); })
+				[
+					ViewportWidget.ToSharedRef()
+				]
+			]
+		]
+
+		// slider
+		+ SVerticalBox::Slot()
+		.AutoHeight()
+		.Padding(0.0f)
+		[
+			FrameSlider.ToSharedRef()
+		]
+	];
+}
+
+END_SLATE_FUNCTION_BUILD_OPTIMIZATION
+
+void UE::RenderPages::Private::SRenderPagesPageViewerLive::OnObjectModified(UObject* Object)
+{
+	if (SelectedPageWeakPtr.IsValid() && (Object == SelectedPageWeakPtr))
+	{
+		// page changed
+		SelectedPageChanged();
+	}
+	else if (const TSharedPtr<IRenderPageCollectionEditor> BlueprintEditor = BlueprintEditorWeakPtr.Pin())
+	{
+		if (Object == BlueprintEditor->GetInstance())
+		{
+			// page collection changed
+			PagesDataChanged();
+		}
+	}
+}
+
+void UE::RenderPages::Private::SRenderPagesPageViewerLive::PagesDataChanged()
+{
+	UpdateViewport();
+	UpdateFrameSlider();
+}
+
+void UE::RenderPages::Private::SRenderPagesPageViewerLive::SelectedPageChanged()
+{
+	SelectedPageWeakPtr = nullptr;
+	if (const TSharedPtr<IRenderPageCollectionEditor> BlueprintEditor = BlueprintEditorWeakPtr.Pin())
+	{
+		if (const TArray<URenderPage*> SelectedPages = BlueprintEditor->GetSelectedRenderPages(); (SelectedPages.Num() == 1))
+		{
+			SelectedPageWeakPtr = SelectedPages[0];
+		}
+	}
+
+	UpdateViewport();
+	UpdateFrameSlider();
+}
+
+void UE::RenderPages::Private::SRenderPagesPageViewerLive::FrameSliderValueChanged(const float NewValue)
+{
+	UpdateViewport();
+	UpdateFrameSlider();
+}
+
+
+void UE::RenderPages::Private::SRenderPagesPageViewerLive::UpdateViewport()
+{
+	if (!ViewportWidget.IsValid())
+	{
+		return;
+	}
+	ViewportWidget->SetVisibility(EVisibility::Hidden);
+
+	if (FrameSlider.IsValid())
+	{
+		if (URenderPage* SelectedPage = SelectedPageWeakPtr.Get(); IsValid(SelectedPage))
+		{
+			if (ULevelSequence* Sequence = SelectedPage->GetSequence(); IsValid(Sequence))
+			{
+				TOptional<double> StartTime = SelectedPage->GetStartTime();
+				TOptional<double> EndTime = SelectedPage->GetEndTime();
+				if (!StartTime.IsSet() || !EndTime.IsSet() || (StartTime.Get(0) > EndTime.Get(0)))
+				{
+					return;
+				}
+
+				if (!ViewportWidget->ShowSequenceFrame(Sequence, FMath::Lerp(StartTime.Get(0), EndTime.Get(0), FrameSlider->GetValue())))
+				{
+					return;
+				}
+
+				ViewportWidget->SetVisibility(EVisibility::Visible);
+			}
+		}
+	}
+}
+
+void UE::RenderPages::Private::SRenderPagesPageViewerLive::UpdateFrameSlider()
+{
+	if (!FrameSlider.IsValid())
+	{
+		return;
+	}
+	FrameSlider->SetVisibility(EVisibility::Hidden);
+
+	if (URenderPage* SelectedPage = SelectedPageWeakPtr.Get(); IsValid(SelectedPage))
+	{
+		if (!FrameSlider->SetFramesText(SelectedPage))
+		{
+			return;
+		}
+
+		FrameSlider->SetVisibility(EVisibility::Visible);
+	}
+}
+
+
+#undef LOCTEXT_NAMESPACE
diff --git a/Engine/Plugins/Experimental/RenderPages/Source/RenderPagesEditor/Private/UI/SRenderPagesPageViewerLive.h b/Engine/Plugins/Experimental/RenderPages/Source/RenderPagesEditor/Private/UI/SRenderPagesPageViewerLive.h
new file mode 100644
--- /dev/null
+++ b/Engine/Plugins/Experimental/RenderPages/Source/RenderPagesEditor/Private/UI/SRenderPagesPageViewerLive.h
@@ -0,0 +1,109 @@
+// Copyright Epic Games, Inc. All Rights Reserved.
+
+#pragma once
+
+#include "CoreMinimal.h"
+#include "LevelSequence.h"
+#include "Widgets/SCompoundWidget.h"
+#include "SEditorViewport.h"
+
+
+class ALevelSequenceActor;
+class SSlider;
+class URenderPage;
+
+namespace UE::RenderPages
+{
+	class IRenderPageCollectionEditor;
+}
+
+namespace UE::RenderPages::Private
+{
+	class SRenderPagesPageViewerFrameSlider;
+}
+
+
+namespace UE::RenderPages::Private
+{
+	/**
+	 * The viewport client for the live page viewer widget.
+	 */
+	class FRenderPagesEditorViewportClient : public FEditorViewportClient
+	{
+	public:
+		explicit FRenderPagesEditorViewportClient(FPreviewScene* PreviewScene, const TWeakPtr<SEditorViewport>& InEditorViewportWidget = nullptr);
+
+	public:
+		//~ Begin FEditorViewportClient Interface
+		virtual EMouseCursor::Type GetCursor(FViewport* InViewport, int32 X, int32 Y) override { return EMouseCursor::Default; }
+		//~ End FEditorViewportClient Interface
+	};
+
+
+	/**
+	 * The viewport for the live page viewer widget.
+	 */
+	class SRenderPagesEditorViewport : public SEditorViewport
+	{
+	public:
+		SLATE_BEGIN_ARGS(SRenderPagesEditorViewport) {}
+		SLATE_END_ARGS()
+
+		virtual void Tick(const FGeometry&, const double, const float) override;
+		void Construct(const FArguments& InArgs);
+		virtual ~SRenderPagesEditorViewport() override;
+
+		bool ShowSequenceFrame(ULevelSequence* InSequence, const float InTime);
+
+	protected:
+		//~ Begin SEditorViewport Interface
+		virtual TSharedRef<FEditorViewportClient> MakeEditorViewportClient() override { return ViewportClient.ToSharedRef(); }
+		virtual void BindCommands() override {}
+		virtual bool SupportsKeyboardFocus() const override { return false; }
+		//~ End SEditorViewport Interface
+
+	private:
+		/** The viewport client. */
+		TSharedPtr<FRenderPagesEditorViewportClient> ViewportClient;
+
+		/** The level sequence actor we spawned to play the sequence of any given page. */
+		TStrongObjectPtr<ALevelSequenceActor> LevelSequenceActor;
+
+		/** The time of the currently playing sequence. */
+		float LevelSequenceTime;
+	};
+
+
+	/**
+	 * A page viewer widget, allows the user to see a live render using a viewport.
+	 */
+	class SRenderPagesPageViewerLive : public SCompoundWidget
+	{
+	public:
+		SLATE_BEGIN_ARGS(SRenderPagesPageViewerLive) {}
+		SLATE_END_ARGS()
+
+		void Construct(const FArguments& InArgs, TSharedPtr<IRenderPageCollectionEditor> InBlueprintEditor);
+
+	private:
+		void OnObjectModified(UObject* Object);
+		void PagesDataChanged();
+		void SelectedPageChanged();
+		void FrameSliderValueChanged(const float NewValue);
+		void UpdateViewport();
+		void UpdateFrameSlider();
+
+	private:
+		/** A reference to the BP Editor that owns this collection. */
+		TWeakPtr<IRenderPageCollectionEditor> BlueprintEditorWeakPtr;
+
+		/** A reference to the job that's currently rendering. */
+		TWeakObjectPtr<URenderPage> SelectedPageWeakPtr;
+
+		/** The scene viewport widget */
+		TSharedPtr<SRenderPagesEditorViewport> ViewportWidget;
+
+		/** The widget that allows the user to select what frame they'd like to see. */
+		TSharedPtr<SRenderPagesPageViewerFrameSlider> FrameSlider;
+	};
+}
diff --git a/Engine/Plugins/Experimental/RenderPages/Source/RenderPagesEditor/Private/UI/SRenderPagesPageViewerPreview.cpp b/Engine/Plugins/Experimental/RenderPages/Source/RenderPagesEditor/Private/UI/SRenderPagesPageViewerPreview.cpp
new file mode 100644
--- /dev/null
+++ b/Engine/Plugins/Experimental/RenderPages/Source/RenderPagesEditor/Private/UI/SRenderPagesPageViewerPreview.cpp
@@ -0,0 +1,425 @@
+// Copyright Epic Games, Inc. All Rights Reserved.
+
+#include "UI/SRenderPagesPageViewerPreview.h"
+#include "UI/Components/SRenderPagesPageViewerFrameSlider.h"
+#include "RenderPage/RenderPageCollection.h"
+#include "IRenderPageCollectionEditor.h"
+#include "IRenderPagesModule.h"
+#include "MoviePipelineOutputSetting.h"
+#include "RenderPage/RenderPageManager.h"
+#include "SlateOptMacros.h"
+#include "Styles/RenderPagesEditorStyle.h"
+#include "Widgets/Input/SButton.h"
+#include "Widgets/Input/SSlider.h"
+#include "Widgets/Layout/SScaleBox.h"
+
+#define LOCTEXT_NAMESPACE "SRenderPagesPageViewerPreview"
+
+
+bool UE::RenderPages::Private::SRenderPagesPageViewerPreview::bHasRenderedSinceAppStart = false;
+
+
+void UE::RenderPages::Private::SRenderPagesPageViewerPreview::Tick(const FGeometry&, const double, const float)
+{
+	if (const TSharedPtr<IRenderPageCollectionEditor> BlueprintEditor = BlueprintEditorWeakPtr.Pin())
+	{
+		if (BlueprintEditor->CanCurrentlyRender() && !CurrentJob)
+		{
+			UpdateRerenderButton();
+			UpdateFrameSlider();
+
+			if (bRunRenderNewPreview)
+			{
+				InternalRenderNewPreview();
+			}
+		}
+	}
+}
+
+BEGIN_SLATE_FUNCTION_BUILD_OPTIMIZATION
+
+void UE::RenderPages::Private::SRenderPagesPageViewerPreview::Construct(const FArguments& InArgs, TSharedPtr<IRenderPageCollectionEditor> InBlueprintEditor)
+{
+	BlueprintEditorWeakPtr = InBlueprintEditor;
+	SelectedPageWeakPtr = nullptr;
+	CurrentJob = nullptr;
+	bRunRenderNewPreview = false;
+	ImageBrushEmpty = FSlateBrush();
+	ImageBrushEmpty.DrawAs = ESlateBrushDrawType::Type::NoDrawType;
+	ImageBrush = FSlateBrush();
+	ImageTexture = nullptr;
+	LastUpdateImageTextureSelectedPageWeakPtr = nullptr;
+
+	SAssignNew(Image, SImage)
+		.Image(&ImageBrushEmpty);
+
+	SAssignNew(ImageBackground, SImage)
+		.Image(&ImageBrushEmpty);
+
+	SAssignNew(FrameSlider, SRenderPagesPageViewerFrameSlider)
+		.Visibility(EVisibility::Hidden)
+		.OnValueChanged(this, &SRenderPagesPageViewerPreview::FrameSliderValueChanged)
+		.OnCaptureEnd(this, &SRenderPagesPageViewerPreview::FrameSliderValueChangedEnd);
+
+	SelectedPageChanged();
+
+	InBlueprintEditor->OnRenderPagesChanged().AddSP(this, &SRenderPagesPageViewerPreview::PagesDataChanged);
+	InBlueprintEditor->OnRenderPagesSelectionChanged().AddSP(this, &SRenderPagesPageViewerPreview::SelectedPageChanged);
+	FCoreUObjectDelegates::OnObjectModified.AddSP(this, &SRenderPagesPageViewerPreview::OnObjectModified);
+
+	ChildSlot
+	[
+		SNew(SVerticalBox)
+		.Visibility_Lambda([this]() -> EVisibility { return (!IsPreviewWidget() && CurrentJob) ? EVisibility::Hidden : EVisibility::Visible; })
+
+		// image
+		+ SVerticalBox::Slot()
+		.FillHeight(1.0f)
+		.Padding(0.0f)
+		[
+			SAssignNew(RerenderButton, SButton)
+			.ContentPadding(0.0f)
+			.ButtonStyle(FRenderPagesEditorStyle::Get(), TEXT("Invisible"))
+			.IsFocusable(false)
+			.OnClicked(this, &SRenderPagesPageViewerPreview::OnClicked)
+			[
+				SNew(SScaleBox)
+				.HAlign(HAlign_Center)
+				.VAlign(VAlign_Center)
+				.Stretch(EStretch::ScaleToFit)
+				.StretchDirection(EStretchDirection::Both)
+				[
+					SNew(SOverlay)
+
+					+ SOverlay::Slot()
+					.HAlign(HAlign_Fill)
+					.VAlign(VAlign_Fill)
+					[
+						ImageBackground.ToSharedRef()
+					]
+
+					+ SOverlay::Slot()
+					.HAlign(HAlign_Fill)
+					.VAlign(VAlign_Fill)
+					[
+						Image.ToSharedRef()
+					]
+				]
+			]
+		]
+
+		// slider
+		+ SVerticalBox::Slot()
+		.AutoHeight()
+		.Padding(0.0f)
+		[
+			FrameSlider.ToSharedRef()
+		]
+	];
+}
+
+END_SLATE_FUNCTION_BUILD_OPTIMIZATION
+
+
+FReply UE::RenderPages::Private::SRenderPagesPageViewerPreview::OnClicked()
+{
+	if (!IsPreviewWidget())
+	{
+		RenderNewPreview();
+	}
+
+	return FReply::Handled();
+}
+
+void UE::RenderPages::Private::SRenderPagesPageViewerPreview::OnObjectModified(UObject* Object)
+{
+	if (SelectedPageWeakPtr.IsValid() && (Object == SelectedPageWeakPtr.Get()))
+	{
+		// page changed
+		SelectedPageChanged();
+	}
+	else if (const TSharedPtr<IRenderPageCollectionEditor> BlueprintEditor = BlueprintEditorWeakPtr.Pin())
+	{
+		if (Object == BlueprintEditor->GetInstance())
+		{
+			// page collection changed
+			PagesDataChanged();
+		}
+	}
+}
+
+
+void UE::RenderPages::Private::SRenderPagesPageViewerPreview::PagesDataChanged()
+{
+	if (IsPreviewWidget())
+	{
+		UpdateImageTexture();
+		RenderNewPreview();
+	}
+	else
+	{
+		UpdateImageTexture();
+	}
+}
+
+void UE::RenderPages::Private::SRenderPagesPageViewerPreview::SelectedPageChanged()
+{
+	SelectedPageWeakPtr = nullptr;
+	if (const TSharedPtr<IRenderPageCollectionEditor> BlueprintEditor = BlueprintEditorWeakPtr.Pin())
+	{
+		if (const TArray<URenderPage*> SelectedPages = BlueprintEditor->GetSelectedRenderPages(); (SelectedPages.Num() == 1))
+		{
+			SelectedPageWeakPtr = SelectedPages[0];
+		}
+	}
+
+	if (IsPreviewWidget())
+	{
+		UpdateImageTexture();
+		RenderNewPreview();
+	}
+	else
+	{
+		UpdateImageTexture();
+	}
+}
+
+void UE::RenderPages::Private::SRenderPagesPageViewerPreview::FrameSliderValueChanged(const float NewValue)
+{
+	if (!IsPreviewWidget())
+	{
+		UpdateImageTexture(false);
+	}
+}
+
+void UE::RenderPages::Private::SRenderPagesPageViewerPreview::FrameSliderValueChangedEnd()
+{
+	if (IsPreviewWidget())
+	{
+		RenderNewPreview();
+	}
+	else
+	{
+		UpdateImageTexture();
+	}
+}
+
+
+void UE::RenderPages::Private::SRenderPagesPageViewerPreview::RenderNewPreview()
+{
+	bRunRenderNewPreview = true;
+}
+
+void UE::RenderPages::Private::SRenderPagesPageViewerPreview::InternalRenderNewPreview()
+{
+	bRunRenderNewPreview = false;
+	InternalRenderNewPreviewOfPage(SelectedPageWeakPtr.Get());
+}
+
+void UE::RenderPages::Private::SRenderPagesPageViewerPreview::InternalRenderNewPreviewOfPage(URenderPage* Page)
+{
+	if (GetTickSpaceGeometry().Size.X <= 0)
+	{
+		// don't render, try again next frame
+		RenderNewPreview();
+		return;
+	}
+
+	if (!IsValid(Page))
+	{
+		SetImageTexture(nullptr);
+		return;
+	}
+
+	if (const TSharedPtr<IRenderPageCollectionEditor> BlueprintEditor = BlueprintEditorWeakPtr.Pin())
+	{
+		double WidgetWidth = FMath::Max(120.0, FMath::Min(GetTickSpaceGeometry().Size.X, GetTickSpaceGeometry().Size.Y * Page->GetOutputAspectRatio()));
+		double RenderResolution = WidgetWidth * (IsPreviewWidget() ? 1.25 : 0.75);// pixels in width
+
+		FRenderPageManagerRenderPreviewFrameArgs JobArgs;
+		if (IsPreviewWidget())
+		{
+			TOptional<int32> SelectedFrame = (!FrameSlider.IsValid() ? TOptional<int32>() : FrameSlider->GetSelectedFrame(Page));
+			if (!SelectedFrame.IsSet())
+			{
+				SetImageTexture(nullptr);
+				return;
+			}
+
+			JobArgs.bHeadless = bHasRenderedSinceAppStart;
+			JobArgs.Frame = SelectedFrame.Get(0);
+		}
+		JobArgs.PageCollection = BlueprintEditor->GetInstance();
+		JobArgs.Page = Page;
+		JobArgs.Resolution = FIntPoint(FMath::RoundToInt32(RenderResolution), FMath::RoundToInt32(RenderResolution / Page->GetOutputAspectRatio()));
+
+		const TSharedPtr<SWidget> BaseThis = AsShared();
+		const TSharedPtr<SRenderPagesPageViewerPreview> This = StaticCastSharedPtr<SRenderPagesPageViewerPreview>(BaseThis);
+		JobArgs.Callback = FRenderPageManagerRenderPreviewFrameArgsCallback::CreateLambda([This, BlueprintEditor](const bool bSuccess)
+		{
+			if (This.IsValid())
+			{
+				This->RenderNewPreviewCallback(bSuccess);
+			}
+			else if (BlueprintEditor.IsValid())
+			{
+				BlueprintEditor->SetPreviewRenderJob(nullptr);
+			}
+		});
+
+		if (URenderPagesMoviePipelineRenderJob* NewJob = IRenderPagesModule::Get().GetManager().RenderPreviewFrame(JobArgs))
+		{
+			CurrentJob = NewJob;
+			BlueprintEditor->SetPreviewRenderJob(CurrentJob);
+			return;
+		}
+	}
+	SetImageTexture(nullptr);
+}
+
+void UE::RenderPages::Private::SRenderPagesPageViewerPreview::RenderNewPreviewCallback(const bool bSuccess)
+{
+	bHasRenderedSinceAppStart = true;
+	UpdateImageTexture();
+
+	CurrentJob = nullptr;
+	if (const TSharedPtr<IRenderPageCollectionEditor> BlueprintEditor = BlueprintEditorWeakPtr.Pin())
+	{
+		BlueprintEditor->SetPreviewRenderJob(CurrentJob);
+	}
+}
+
+void UE::RenderPages::Private::SRenderPagesPageViewerPreview::UpdateImageTexture(const bool bForce)
+{
+	if (FrameSlider.IsValid())
+	{
+		if (URenderPage* SelectedPage = SelectedPageWeakPtr.Get(); IsValid(SelectedPage))
+		{
+			if (!bForce && LastUpdateImageTextureSelectedPageWeakPtr.IsValid() && (LastUpdateImageTextureSelectedPageWeakPtr == SelectedPage))
+			{
+				if (IsPreviewWidget() || (LastUpdateImageTextureFrame == FrameSlider->GetSelectedFrame(SelectedPage)))
+				{
+					return; // nothing changed
+				}
+			}
+			LastUpdateImageTextureSelectedPageWeakPtr = SelectedPage;
+
+			if (IsPreviewWidget())
+			{
+				if (!FrameSlider->GetSelectedFrame(SelectedPage).IsSet())
+				{
+					SetImageTexture(nullptr);
+					return;
+				}
+				SetImageTexture(IRenderPagesModule::Get().GetManager().GetSingleRenderedPreviewFrame(SelectedPage));
+				return;
+			}
+
+			TOptional<int32> Frame = FrameSlider->GetSelectedFrame(SelectedPage);
+			LastUpdateImageTextureFrame = Frame;
+			if (!Frame.IsSet())
+			{
+				SetImageTexture(nullptr);
+				return;
+			}
+			SetImageTexture(IRenderPagesModule::Get().GetManager().GetRenderedPreviewFrame(SelectedPage, Frame.Get(0)));
+			return;
+		}
+	}
+	SetImageTexture(nullptr);
+	LastUpdateImageTextureSelectedPageWeakPtr = nullptr;
+}
+
+void UE::RenderPages::Private::SRenderPagesPageViewerPreview::SetImageTexture(UTexture2D* Texture)
+{
+	{// cleanup >>
+		if (Image.IsValid())
+		{
+			Image->SetImage(&ImageBrushEmpty);
+			ImageBackground->SetImage(&ImageBrushEmpty);
+		}
+		ImageBrush.SetResourceObject(nullptr);
+		ImageBrush.SetImageSize(FVector2D(0, 0));
+
+		if (IsValid(ImageTexture))
+		{
+			ImageTexture->RemoveFromRoot();
+		}
+		ImageTexture = nullptr;
+	}// cleanup <<
+
+	{// set new texture >>
+		if (IsValid(Texture) && Texture->GetResource())
+		{
+			ImageTexture = Texture;
+			ImageTexture->AddToRoot();
+
+			if (Image.IsValid())
+			{
+				static constexpr double PreviewTabAspectRatio = 1.96875;
+
+				double ImageWidth = ImageTexture->GetResource()->GetSizeX();
+				double ImageHeight = ImageTexture->GetResource()->GetSizeY();
+				const double ImageAspectRatio = ImageWidth / ImageHeight;
+				if (ImageAspectRatio > PreviewTabAspectRatio)
+				{
+					ImageWidth = 1280.0;
+					ImageHeight = ImageWidth / ImageAspectRatio;
+				}
+				else
+				{
+					ImageHeight = 1280.0 / PreviewTabAspectRatio;
+					ImageWidth = ImageHeight * ImageAspectRatio;
+				}
+
+				ImageBrush = FSlateBrush();
+				ImageBrush.DrawAs = ESlateBrushDrawType::Type::Image;
+				ImageBrush.ImageType = ESlateBrushImageType::Type::FullColor;
+				ImageBrush.SetResourceObject(ImageTexture);
+				ImageBrush.SetImageSize(FVector2D(ImageWidth, ImageHeight));
+				Image->SetImage(&ImageBrush);
+				ImageBackground->SetImage(FCoreStyle::Get().GetBrush("Checkerboard"));
+			}
+		}
+	}// set new texture <<
+}
+
+
+void UE::RenderPages::Private::SRenderPagesPageViewerPreview::UpdateRerenderButton()
+{
+	if (!RerenderButton.IsValid())
+	{
+		return;
+	}
+	const bool bIsUsable = (!IsPreviewWidget() && SelectedPageWeakPtr.IsValid());
+
+	RerenderButton->SetButtonStyle(&FRenderPagesEditorStyle::Get().GetWidgetStyle<FButtonStyle>(bIsUsable ? TEXT("HoverHintOnly") : TEXT("Invisible")));
+	RerenderButton->SetCursor(bIsUsable ? EMouseCursor::Type::Hand : EMouseCursor::Type::Default);
+}
+
+void UE::RenderPages::Private::SRenderPagesPageViewerPreview::UpdateFrameSlider()
+{
+	if (!FrameSlider.IsValid())
+	{
+		return;
+	}
+	FrameSlider->SetVisibility(EVisibility::Hidden);
+
+	if (URenderPage* SelectedPage = SelectedPageWeakPtr.Get(); IsValid(SelectedPage))
+	{
+		if (IsPreviewWidget() && !FrameSlider->GetSelectedSequenceFrame(SelectedPage).IsSet())
+		{
+			return;
+		}
+
+		if (!FrameSlider->SetFramesText(SelectedPage))
+		{
+			return;
+		}
+
+		FrameSlider->SetVisibility(EVisibility::Visible);
+	}
+}
+
+
+#undef LOCTEXT_NAMESPACE
diff --git a/Engine/Plugins/Experimental/RenderPages/Source/RenderPagesEditor/Private/UI/SRenderPagesPageViewerPreview.h b/Engine/Plugins/Experimental/RenderPages/Source/RenderPagesEditor/Private/UI/SRenderPagesPageViewerPreview.h
new file mode 100644
--- /dev/null
+++ b/Engine/Plugins/Experimental/RenderPages/Source/RenderPagesEditor/Private/UI/SRenderPagesPageViewerPreview.h
@@ -0,0 +1,116 @@
+// Copyright Epic Games, Inc. All Rights Reserved.
+
+#pragma once
+
+#include "CoreMinimal.h"
+#include "Widgets/SCompoundWidget.h"
+
+
+class SButton;
+class SImage;
+class SSlider;
+class URenderPage;
+class URenderPagesMoviePipelineRenderJob;
+class UTexture2D;
+
+
+namespace UE::RenderPages
+{
+	class IRenderPageCollectionEditor;
+}
+
+namespace UE::RenderPages::Private
+{
+	class SRenderPagesPageViewerFrameSlider;
+	struct FRenderPageManagerRenderPreviewFrameResult;
+}
+
+
+namespace UE::RenderPages::Private
+{
+	/**
+	 * A page viewer widget, allows the user to render a single frame of a page in low-resolution and afterwards see it in the editor.
+	 */
+	class SRenderPagesPageViewerPreview : public SCompoundWidget
+	{
+	public:
+		SLATE_BEGIN_ARGS(SRenderPagesPageViewerPreview) {}
+		SLATE_END_ARGS()
+
+		virtual void Tick(const FGeometry&, const double, const float) override;
+		void Construct(const FArguments& InArgs, TSharedPtr<IRenderPageCollectionEditor> InBlueprintEditor);
+
+		/** Queues a new render action (if it isn't yet queued already). This is basically the Refresh() function of this widget. */
+		void RenderNewPreview();
+
+		/** Retrieves the rendered output from disk and shows it in the UI. */
+		void UpdateImageTexture(const bool bForce = true);
+
+		/** Shows the given texture, this is normally the output of a render. */
+		void SetImageTexture(UTexture2D* Texture);
+
+	protected:
+		/** Returns true if this is the preview widget (1 frame), returns false if this is the rendered widgets (all frames). Override this function in order to change the value it returns. */
+		virtual bool IsPreviewWidget() const { return true; }
+
+	private:
+		FReply OnClicked();
+		void OnObjectModified(UObject* Object);
+		void PagesDataChanged();
+		void SelectedPageChanged();
+		void FrameSliderValueChanged(const float NewValue);
+		void FrameSliderValueChangedEnd();
+
+		void InternalRenderNewPreview();
+		void InternalRenderNewPreviewOfPage(URenderPage* Page);
+		void RenderNewPreviewCallback(const bool bSuccess);
+
+		void UpdateRerenderButton();
+		void UpdateFrameSlider();
+
+
+	private:
+		/** A reference to the BP Editor that owns this collection. */
+		TWeakPtr<IRenderPageCollectionEditor> BlueprintEditorWeakPtr;
+
+		/** A reference to the job that's currently rendering. */
+		TObjectPtr<URenderPagesMoviePipelineRenderJob> CurrentJob;
+
+		/** A reference to the job that's currently rendering. */
+		TWeakObjectPtr<URenderPage> SelectedPageWeakPtr;
+
+		/** True if it should call RenderNewPreview() next frame. */
+		bool bRunRenderNewPreview;
+
+	private:
+		/** True if it has rendered before since the start of this application. Used for not hiding the rendering popup during the first render (since the first render can take a lot longer due to having to compile shaders etc). */
+		static bool bHasRenderedSinceAppStart;
+
+		/** The widget that allows the user to select what frame they'd like to see. */
+		TSharedPtr<SRenderPagesPageViewerFrameSlider> FrameSlider;
+
+		/** The button which can be clicked to rerender the preview. */
+		TSharedPtr<SButton> RerenderButton;
+
+		/** The widget that contains the image. */
+		TSharedPtr<SImage> Image;
+
+		/** The widget that contains the background of the image. */
+		TSharedPtr<SImage> ImageBackground;
+
+		/** The brush of the image, always empty. */
+		FSlateBrush ImageBrushEmpty;
+
+		/** The brush of the image. */
+		FSlateBrush ImageBrush;
+
+		/** The texture of the image. */
+		TObjectPtr<UTexture2D> ImageTexture;
+
+		/** The last page used for the LastUpdateImageTexture function. */
+		TWeakObjectPtr<URenderPage> LastUpdateImageTextureSelectedPageWeakPtr;
+
+		/** The last frame used for the LastUpdateImageTexture function. */
+		TOptional<int32> LastUpdateImageTextureFrame;
+	};
+}
diff --git a/Engine/Plugins/Experimental/RenderPages/Source/RenderPagesEditor/Private/UI/SRenderPagesPageViewerRendered.cpp b/Engine/Plugins/Experimental/RenderPages/Source/RenderPagesEditor/Private/UI/SRenderPagesPageViewerRendered.cpp
new file mode 100644
--- /dev/null
+++ b/Engine/Plugins/Experimental/RenderPages/Source/RenderPagesEditor/Private/UI/SRenderPagesPageViewerRendered.cpp
@@ -0,0 +1,3 @@
+// Copyright Epic Games, Inc. All Rights Reserved.
+
+#include "UI/SRenderPagesPageViewerRendered.h"
diff --git a/Engine/Plugins/Experimental/RenderPages/Source/RenderPagesEditor/Private/UI/SRenderPagesPageViewerRendered.h b/Engine/Plugins/Experimental/RenderPages/Source/RenderPagesEditor/Private/UI/SRenderPagesPageViewerRendered.h
new file mode 100644
--- /dev/null
+++ b/Engine/Plugins/Experimental/RenderPages/Source/RenderPagesEditor/Private/UI/SRenderPagesPageViewerRendered.h
@@ -0,0 +1,18 @@
+// Copyright Epic Games, Inc. All Rights Reserved.
+
+#pragma once
+
+#include "UI/SRenderPagesPageViewerPreview.h"
+
+
+namespace UE::RenderPages::Private
+{
+	/**
+	 * A page viewer widget, allows the user to render a page in low-resolution and afterwards scrub through the outputted frames of it in the editor.
+	 */
+	class SRenderPagesPageViewerRendered : public SRenderPagesPageViewerPreview
+	{
+	protected:
+		virtual bool IsPreviewWidget() const override { return false; }
+	};
+}
diff --git a/Engine/Plugins/Experimental/RenderPages/Source/RenderPagesEditor/Private/UI/SRenderPagesProps.cpp b/Engine/Plugins/Experimental/RenderPages/Source/RenderPagesEditor/Private/UI/SRenderPagesProps.cpp
new file mode 100644
--- /dev/null
+++ b/Engine/Plugins/Experimental/RenderPages/Source/RenderPagesEditor/Private/UI/SRenderPagesProps.cpp
@@ -0,0 +1,74 @@
+// Copyright Epic Games, Inc. All Rights Reserved.
+
+#include "UI/SRenderPagesProps.h"
+#include "UI/SRenderPagesPropsBase.h"
+#include "RenderPage/RenderPageCollection.h"
+#include "IRenderPageCollectionEditor.h"
+#include "IRenderPagesEditorModule.h"
+#include "SlateOptMacros.h"
+#include "Widgets/Layout/SBorder.h"
+
+#define LOCTEXT_NAMESPACE "SRenderPagesProps"
+
+
+void UE::RenderPages::Private::SRenderPagesProps::Tick(const FGeometry&, const double, const float)
+{
+	if (const TSharedPtr<IRenderPageCollectionEditor> BlueprintEditor = BlueprintEditorWeakPtr.Pin())
+	{
+		if (URenderPageCollection* Collection = BlueprintEditor->GetInstance(); IsValid(Collection))
+		{
+			if (WidgetPropsSourceWeakPtr != Collection->GetPropsSource())
+			{
+				Refresh();
+			}
+		}
+	}
+}
+
+BEGIN_SLATE_FUNCTION_BUILD_OPTIMIZATION
+
+void UE::RenderPages::Private::SRenderPagesProps::Construct(const FArguments& InArgs, TSharedPtr<IRenderPageCollectionEditor> InBlueprintEditor)
+{
+	BlueprintEditorWeakPtr = InBlueprintEditor;
+
+	SAssignNew(WidgetContainer, SBorder)
+		.Padding(8.0f)
+		.BorderImage(new FSlateNoResource());
+
+	Refresh();
+	InBlueprintEditor->OnRenderPagesChanged().AddSP(this, &SRenderPagesProps::Refresh);
+
+	ChildSlot
+	[
+		WidgetContainer.ToSharedRef()
+	];
+}
+
+END_SLATE_FUNCTION_BUILD_OPTIMIZATION
+
+void UE::RenderPages::Private::SRenderPagesProps::Refresh()
+{
+	if (!WidgetContainer.IsValid())
+	{
+		return;
+	}
+	WidgetContainer->ClearContent();
+
+	if (const TSharedPtr<IRenderPageCollectionEditor> BlueprintEditor = BlueprintEditorWeakPtr.Pin())
+	{
+		if (URenderPageCollection* Collection = BlueprintEditor->GetInstance(); IsValid(Collection))
+		{
+			WidgetPropsSourceWeakPtr = Collection->GetPropsSource();
+			if (URenderPagePropsSourceBase* WidgetPropsSource = WidgetPropsSourceWeakPtr.Get(); IsValid(WidgetPropsSource))
+			{
+				if (const TSharedPtr<SRenderPagesPropsBase> Widget = IRenderPagesEditorModule::Get().CreatePropsSourceWidget(WidgetPropsSource, BlueprintEditor))
+				{
+					WidgetContainer->SetContent(Widget.ToSharedRef());
+				}
+			}
+		}
+	}
+}
+
+
+#undef LOCTEXT_NAMESPACE
diff --git a/Engine/Plugins/Experimental/RenderPages/Source/RenderPagesEditor/Private/UI/SRenderPagesProps.h b/Engine/Plugins/Experimental/RenderPages/Source/RenderPagesEditor/Private/UI/SRenderPagesProps.h
new file mode 100644
--- /dev/null
+++ b/Engine/Plugins/Experimental/RenderPages/Source/RenderPagesEditor/Private/UI/SRenderPagesProps.h
@@ -0,0 +1,46 @@
+// Copyright Epic Games, Inc. All Rights Reserved.
+
+#pragma once
+
+#include "CoreMinimal.h"
+#include "Widgets/SCompoundWidget.h"
+
+
+class SBorder;
+class URenderPagePropsSourceBase;
+
+namespace UE::RenderPages
+{
+	class IRenderPageCollectionEditor;
+}
+
+
+namespace UE::RenderPages::Private
+{
+	/**
+	 * A widget with which the user can modify the props (like the remote control field values) of the selected render page.
+	 * Can only modify the props of 1 render page at a time, this widget will show nothing when 0 or 2+ render pages are selected.
+	 */
+	class SRenderPagesProps : public SCompoundWidget
+	{
+	public:
+		SLATE_BEGIN_ARGS(SRenderPagesProps) {}
+		SLATE_END_ARGS()
+
+		virtual void Tick(const FGeometry&, const double, const float) override;
+		void Construct(const FArguments& InArgs, TSharedPtr<IRenderPageCollectionEditor> InBlueprintEditor);
+
+		/** Refreshes the content of this widget. */
+		void Refresh();
+
+	private:
+		/** A reference to the BP Editor that owns this collection. */
+		TWeakPtr<IRenderPageCollectionEditor> BlueprintEditorWeakPtr;
+
+		/** The widget that lists the property rows. */
+		TSharedPtr<SBorder> WidgetContainer;
+
+		/** The props source that's being shown in this widget. */
+		TWeakObjectPtr<URenderPagePropsSourceBase> WidgetPropsSourceWeakPtr;
+	};
+}
diff --git a/Engine/Plugins/Experimental/RenderPages/Source/RenderPagesEditor/Private/UI/SRenderPagesPropsBase.h b/Engine/Plugins/Experimental/RenderPages/Source/RenderPagesEditor/Private/UI/SRenderPagesPropsBase.h
new file mode 100644
--- /dev/null
+++ b/Engine/Plugins/Experimental/RenderPages/Source/RenderPagesEditor/Private/UI/SRenderPagesPropsBase.h
@@ -0,0 +1,15 @@
+// Copyright Epic Games, Inc. All Rights Reserved.
+
+#pragma once
+
+#include "CoreMinimal.h"
+#include "Widgets/SCompoundWidget.h"
+
+
+namespace UE::RenderPages::Private
+{
+	/**
+	 * The base class of the page props widget.
+	 */
+	class SRenderPagesPropsBase : public SCompoundWidget {};
+}
diff --git a/Engine/Plugins/Experimental/RenderPages/Source/RenderPagesEditor/Private/UI/SRenderPagesPropsLocal.cpp b/Engine/Plugins/Experimental/RenderPages/Source/RenderPagesEditor/Private/UI/SRenderPagesPropsLocal.cpp
new file mode 100644
--- /dev/null
+++ b/Engine/Plugins/Experimental/RenderPages/Source/RenderPagesEditor/Private/UI/SRenderPagesPropsLocal.cpp
@@ -0,0 +1,22 @@
+// Copyright Epic Games, Inc. All Rights Reserved.
+
+#include "UI/SRenderPagesPropsLocal.h"
+#include "RenderPage/RenderPageCollection.h"
+#include "IRenderPageCollectionEditor.h"
+#include "SlateOptMacros.h"
+
+#define LOCTEXT_NAMESPACE "SRenderPagesPropsLocal"
+
+
+BEGIN_SLATE_FUNCTION_BUILD_OPTIMIZATION
+
+void UE::RenderPages::Private::SRenderPagesPropsLocal::Construct(const FArguments& InArgs, TSharedPtr<IRenderPageCollectionEditor> InBlueprintEditor, URenderPagePropsSourceLocal* InPropsSource)
+{
+	BlueprintEditorWeakPtr = InBlueprintEditor;
+	PropsSource = InPropsSource;
+}
+
+END_SLATE_FUNCTION_BUILD_OPTIMIZATION
+
+
+#undef LOCTEXT_NAMESPACE
diff --git a/Engine/Plugins/Experimental/RenderPages/Source/RenderPagesEditor/Private/UI/SRenderPagesPropsLocal.h b/Engine/Plugins/Experimental/RenderPages/Source/RenderPagesEditor/Private/UI/SRenderPagesPropsLocal.h
new file mode 100644
--- /dev/null
+++ b/Engine/Plugins/Experimental/RenderPages/Source/RenderPagesEditor/Private/UI/SRenderPagesPropsLocal.h
@@ -0,0 +1,36 @@
+// Copyright Epic Games, Inc. All Rights Reserved.
+
+#pragma once
+
+#include "CoreMinimal.h"
+#include "UI/SRenderPagesPropsBase.h"
+#include "RenderPage/RenderPagePropsSource.h"
+
+
+namespace UE::RenderPages
+{
+	class IRenderPageCollectionEditor;
+}
+
+
+namespace UE::RenderPages::Private
+{
+	/**
+	 * The page props implementation for local properties.
+	 */
+	class SRenderPagesPropsLocal : public SRenderPagesPropsBase
+	{
+	public:
+		SLATE_BEGIN_ARGS(SRenderPagesPropsLocal) {}
+		SLATE_END_ARGS()
+
+		void Construct(const FArguments& InArgs, TSharedPtr<IRenderPageCollectionEditor> InBlueprintEditor, URenderPagePropsSourceLocal* InPropsSource);
+
+	private:
+		/** A reference to the BP Editor that owns this collection. */
+		TWeakPtr<IRenderPageCollectionEditor> BlueprintEditorWeakPtr;
+
+		/** The props source control. */
+		TObjectPtr<URenderPagePropsSourceLocal> PropsSource;
+	};
+}
diff --git a/Engine/Plugins/Experimental/RenderPages/Source/RenderPagesEditor/Private/UI/SRenderPagesPropsRemoteControl.cpp b/Engine/Plugins/Experimental/RenderPages/Source/RenderPagesEditor/Private/UI/SRenderPagesPropsRemoteControl.cpp
new file mode 100644
--- /dev/null
+++ b/Engine/Plugins/Experimental/RenderPages/Source/RenderPagesEditor/Private/UI/SRenderPagesPropsRemoteControl.cpp
@@ -0,0 +1,199 @@
+// Copyright Epic Games, Inc. All Rights Reserved.
+
+#include "UI/SRenderPagesPropsRemoteControl.h"
+#include "UI/Components/SRenderPagesRemoteControlTreeNode.h"
+#include "UI/Components/SRenderPagesRemoteControlField.h"
+#include "IRenderPageCollectionEditor.h"
+#include "IRenderPagesModule.h"
+#include "RenderPage/RenderPageCollection.h"
+#include "RenderPage/RenderPageManager.h"
+#include "SlateOptMacros.h"
+
+#define LOCTEXT_NAMESPACE "SRenderPagesPropsRemoteControl"
+
+
+BEGIN_SLATE_FUNCTION_BUILD_OPTIMIZATION
+
+void UE::RenderPages::Private::SRenderPagesPropsRemoteControl::Construct(const FArguments& InArgs, TSharedPtr<IRenderPageCollectionEditor> InBlueprintEditor, URenderPagePropsSourceRemoteControl* InPropsSource)
+{
+	BlueprintEditorWeakPtr = InBlueprintEditor;
+	PropsSource = InPropsSource;
+
+	SAssignNew(RowWidgetsContainer, SVerticalBox);
+	UpdateStoredValuesAndRefresh();
+
+	InBlueprintEditor->OnRenderPagesSelectionChanged().AddSP(this, &SRenderPagesPropsRemoteControl::Refresh);
+	if (IsValid(PropsSource))
+	{
+		if (TObjectPtr<URemoteControlPreset> Preset = PropsSource->GetProps()->GetRemoteControlPreset())
+		{
+			Preset->OnEntityExposed().AddSP(this, &SRenderPagesPropsRemoteControl::OnRemoteControlEntitiesExposed);
+			Preset->OnEntityUnexposed().AddSP(this, &SRenderPagesPropsRemoteControl::OnRemoteControlEntitiesUnexposed);
+			Preset->OnEntitiesUpdated().AddSP(this, &SRenderPagesPropsRemoteControl::OnRemoteControlEntitiesUpdated);
+			Preset->OnExposedPropertiesModified().AddSP(this, &SRenderPagesPropsRemoteControl::OnRemoteControlExposedPropertiesModified);
+		}
+	}
+
+	ChildSlot
+	[
+		SNew(SBorder)
+		.Padding(8.0f)
+		.BorderImage(new FSlateNoResource())
+		[
+			RowWidgetsContainer.ToSharedRef()
+		]
+	];
+}
+
+END_SLATE_FUNCTION_BUILD_OPTIMIZATION
+
+
+void UE::RenderPages::Private::SRenderPagesPropsRemoteControl::UpdateStoredValuesAndRefresh()
+{
+	if (const TSharedPtr<IRenderPageCollectionEditor> BlueprintEditor = BlueprintEditorWeakPtr.Pin())
+	{
+		IRenderPagesModule::Get().GetManager().UpdatePagesPropValues(BlueprintEditor->GetInstance());
+		Refresh();
+	}
+}
+
+void UE::RenderPages::Private::SRenderPagesPropsRemoteControl::Refresh()
+{
+	if (!RowWidgetsContainer.IsValid())
+	{
+		return;
+	}
+	RowWidgetsContainer->ClearChildren();
+
+	if (const TSharedPtr<IRenderPageCollectionEditor> BlueprintEditor = BlueprintEditorWeakPtr.Pin())
+	{
+		if (BlueprintEditor->IsCurrentlyRenderingOrPlaying())
+		{
+			return;
+		}
+
+		TArray<TSharedRef<SWidget>> RowWidgets;
+		if (IsValid(PropsSource))
+		{
+			URenderPagePropsRemoteControl* Props = PropsSource->GetProps();
+			for (URenderPagePropRemoteControl* Prop : Props->GetAllCasted())
+			{
+				TArray<uint8> PropData;
+				if (!GetSelectedPageFieldValue(Prop->GetRemoteControlEntity(), PropData))
+				{
+					continue;
+				}
+				if (!Prop->SetValue(PropData))
+				{
+					continue;
+				}
+
+				if (URemoteControlPreset* Preset = Props->GetRemoteControlPreset())
+				{
+					FRenderPagesRemoteControlGenerateWidgetArgs Args;
+					Args.Preset = Preset;
+					Args.Entity = Prop->GetRemoteControlEntity();
+					Args.ColumnSizeData.LeftColumnWidth = 0.3;
+					Args.ColumnSizeData.RightColumnWidth = 0.7;
+					RowWidgets.Add(SRenderPagesRemoteControlField::MakeInstance(Args).ToSharedRef());
+				}
+			}
+		}
+
+		for (TSharedRef<SWidget> RowWidget : RowWidgets)
+		{
+			RowWidgetsContainer->AddSlot()
+				.Padding(0.0f)
+				.AutoHeight()
+				[
+					RowWidget
+				];
+		}
+	}
+}
+
+
+void UE::RenderPages::Private::SRenderPagesPropsRemoteControl::OnRemoteControlExposedPropertiesModified(URemoteControlPreset* Preset, const TSet<FGuid>& ModifiedProperties)
+{
+	if (!IsValid(Preset))
+	{
+		return;
+	}
+	if (const TSharedPtr<IRenderPageCollectionEditor> BlueprintEditor = BlueprintEditorWeakPtr.Pin())
+	{
+		if (BlueprintEditor->IsCurrentlyRenderingOrPlaying())
+		{
+			return;
+		}
+
+		bool bModified = false;
+		for (const FGuid& Id : ModifiedProperties)
+		{
+			if (const TSharedPtr<FRemoteControlEntity> Entity = Preset->GetExposedEntity<FRemoteControlEntity>(Id).Pin())
+			{
+				TArray<uint8> BinaryArray;
+				if (!URenderPagePropRemoteControl::GetValueOfEntity(Entity, BinaryArray))
+				{
+					continue;
+				}
+				TArray<uint8> StoredBinaryArray;
+				if (!GetSelectedPageFieldValue(Entity, StoredBinaryArray))
+				{
+					continue;
+				}
+				if (BinaryArray == StoredBinaryArray)
+				{
+					continue;
+				}
+
+				if (!SetSelectedPageFieldValue(Entity, BinaryArray))
+				{
+					continue;
+				}
+				bModified = true;
+			}
+		}
+
+		if (bModified)
+		{
+			BlueprintEditor->MarkAsModified();
+			BlueprintEditor->OnRenderPagesChanged().Broadcast();
+			Refresh();
+		}
+	}
+}
+
+
+URenderPage* UE::RenderPages::Private::SRenderPagesPropsRemoteControl::GetSelectedPage()
+{
+	if (const TSharedPtr<IRenderPageCollectionEditor> BlueprintEditor = BlueprintEditorWeakPtr.Pin())
+	{
+		if (const TArray<URenderPage*> SelectedPages = BlueprintEditor->GetSelectedRenderPages(); (SelectedPages.Num() == 1))
+		{
+			return SelectedPages[0];
+		}
+	}
+	return nullptr;
+}
+
+bool UE::RenderPages::Private::SRenderPagesPropsRemoteControl::GetSelectedPageFieldValue(const TSharedPtr<FRemoteControlEntity>& RemoteControlEntity, TArray<uint8>& OutBinaryArray)
+{
+	OutBinaryArray.Empty();
+	if (URenderPage* Page = GetSelectedPage(); IsValid(Page))
+	{
+		return Page->GetRemoteControlValue(RemoteControlEntity, OutBinaryArray);
+	}
+	return false;
+}
+
+bool UE::RenderPages::Private::SRenderPagesPropsRemoteControl::SetSelectedPageFieldValue(const TSharedPtr<FRemoteControlEntity>& RemoteControlEntity, const TArray<uint8>& BinaryArray)
+{
+	if (URenderPage* Page = GetSelectedPage(); IsValid(Page))
+	{
+		return Page->SetRemoteControlValue(RemoteControlEntity, BinaryArray);
+	}
+	return false;
+}
+
+
+#undef LOCTEXT_NAMESPACE
diff --git a/Engine/Plugins/Experimental/RenderPages/Source/RenderPagesEditor/Private/UI/SRenderPagesPropsRemoteControl.h b/Engine/Plugins/Experimental/RenderPages/Source/RenderPagesEditor/Private/UI/SRenderPagesPropsRemoteControl.h
new file mode 100644
--- /dev/null
+++ b/Engine/Plugins/Experimental/RenderPages/Source/RenderPagesEditor/Private/UI/SRenderPagesPropsRemoteControl.h
@@ -0,0 +1,64 @@
+// Copyright Epic Games, Inc. All Rights Reserved.
+
+#pragma once
+
+#include "CoreMinimal.h"
+#include "RemoteControlPreset.h"
+#include "UI/SRenderPagesPropsBase.h"
+#include "RenderPage/RenderPagePropsSource.h"
+
+
+class URenderPage;
+
+namespace UE::RenderPages
+{
+	class IRenderPageCollectionEditor;
+}
+
+
+namespace UE::RenderPages::Private
+{
+	/**
+	 * The page props implementation for remote control fields.
+	 */
+	class SRenderPagesPropsRemoteControl : public SRenderPagesPropsBase
+	{
+	public:
+		SLATE_BEGIN_ARGS(SRenderPagesPropsRemoteControl) {}
+		SLATE_END_ARGS()
+
+		void Construct(const FArguments& InArgs, TSharedPtr<IRenderPageCollectionEditor> InBlueprintEditor, URenderPagePropsSourceRemoteControl* InPropsSource);
+
+		/** Obtain the latest prop values and refreshes the content of this widget. */
+		void UpdateStoredValuesAndRefresh();
+
+		/** Refreshes the content of this widget. */
+		void Refresh();
+
+	private:
+		void OnRemoteControlEntitiesExposed(URemoteControlPreset* Preset, const FGuid& EntityId) { UpdateStoredValuesAndRefresh(); }
+		void OnRemoteControlEntitiesUnexposed(URemoteControlPreset* Preset, const FGuid& EntityId) { UpdateStoredValuesAndRefresh(); }
+		void OnRemoteControlEntitiesUpdated(URemoteControlPreset* Preset, const TSet<FGuid>& ModifiedEntities) { UpdateStoredValuesAndRefresh(); }
+		void OnRemoteControlExposedPropertiesModified(URemoteControlPreset* Preset, const TSet<FGuid>& ModifiedProperties);
+
+	private:
+		/** Returns the currently selected page if 1 page is currently selected, returns nullptr otherwise. */
+		URenderPage* GetSelectedPage();
+
+		/** Obtains the value (as bytes) of the given prop (the given remote control entity), returns true if it succeeded, returns false otherwise. */
+		bool GetSelectedPageFieldValue(const TSharedPtr<FRemoteControlEntity>& RemoteControlEntity, TArray<uint8>& OutBinaryArray);
+
+		/** Sets the value of the given prop (the given remote control entity) with the given value (as bytes), returns true if it succeeded, returns false otherwise. */
+		bool SetSelectedPageFieldValue(const TSharedPtr<FRemoteControlEntity>& RemoteControlEntity, const TArray<uint8>& BinaryArray);
+
+	private:
+		/** A reference to the BP Editor that owns this collection. */
+		TWeakPtr<IRenderPageCollectionEditor> BlueprintEditorWeakPtr;
+
+		/** The props source control. */
+		TObjectPtr<URenderPagePropsSourceRemoteControl> PropsSource;
+
+		/** The widget that lists the property rows. */
+		TSharedPtr<SVerticalBox> RowWidgetsContainer;
+	};
+}
diff --git a/Engine/Plugins/Experimental/RenderPages/Source/RenderPagesEditor/Private/UI/Tabs/SRenderPagesCollectionPropertiesTab.cpp b/Engine/Plugins/Experimental/RenderPages/Source/RenderPagesEditor/Private/UI/Tabs/SRenderPagesCollectionPropertiesTab.cpp
new file mode 100644
--- /dev/null
+++ b/Engine/Plugins/Experimental/RenderPages/Source/RenderPagesEditor/Private/UI/Tabs/SRenderPagesCollectionPropertiesTab.cpp
@@ -0,0 +1,24 @@
+// Copyright Epic Games, Inc. All Rights Reserved.
+
+#include "UI/Tabs/SRenderPagesCollectionPropertiesTab.h"
+#include "UI/SRenderPagesCollection.h"
+#include "IRenderPageCollectionEditor.h"
+#include "SlateOptMacros.h"
+
+#define LOCTEXT_NAMESPACE "SRenderPagesCollectionPropertiesTab"
+
+
+BEGIN_SLATE_FUNCTION_BUILD_OPTIMIZATION
+
+void UE::RenderPages::Private::SRenderPagesCollectionPropertiesTab::Construct(const FArguments& InArgs, TSharedPtr<IRenderPageCollectionEditor> InBlueprintEditor)
+{
+	ChildSlot
+	[
+		SNew(SRenderPagesCollection, InBlueprintEditor)
+	];
+}
+
+END_SLATE_FUNCTION_BUILD_OPTIMIZATION
+
+
+#undef LOCTEXT_NAMESPACE
diff --git a/Engine/Plugins/Experimental/RenderPages/Source/RenderPagesEditor/Private/UI/Tabs/SRenderPagesCollectionPropertiesTab.h b/Engine/Plugins/Experimental/RenderPages/Source/RenderPagesEditor/Private/UI/Tabs/SRenderPagesCollectionPropertiesTab.h
new file mode 100644
--- /dev/null
+++ b/Engine/Plugins/Experimental/RenderPages/Source/RenderPagesEditor/Private/UI/Tabs/SRenderPagesCollectionPropertiesTab.h
@@ -0,0 +1,28 @@
+// Copyright Epic Games, Inc. All Rights Reserved.
+
+#pragma once
+
+#include "CoreMinimal.h"
+#include "Widgets/SCompoundWidget.h"
+
+
+namespace UE::RenderPages
+{
+	class IRenderPageCollectionEditor;
+}
+
+
+namespace UE::RenderPages::Private
+{
+	/**
+	 * The collection properties tab.
+	 */
+	class SRenderPagesCollectionPropertiesTab : public SCompoundWidget
+	{
+	public:
+		SLATE_BEGIN_ARGS(SRenderPagesCollectionPropertiesTab) {}
+		SLATE_END_ARGS()
+
+		void Construct(const FArguments& InArgs, TSharedPtr<IRenderPageCollectionEditor> InBlueprintEditor);
+	};
+}
diff --git a/Engine/Plugins/Experimental/RenderPages/Source/RenderPagesEditor/Private/UI/Tabs/SRenderPagesPageListTab.cpp b/Engine/Plugins/Experimental/RenderPages/Source/RenderPagesEditor/Private/UI/Tabs/SRenderPagesPageListTab.cpp
new file mode 100644
--- /dev/null
+++ b/Engine/Plugins/Experimental/RenderPages/Source/RenderPagesEditor/Private/UI/Tabs/SRenderPagesPageListTab.cpp
@@ -0,0 +1,24 @@
+// Copyright Epic Games, Inc. All Rights Reserved.
+
+#include "UI/Tabs/SRenderPagesPageListTab.h"
+#include "UI/SRenderPagesPageList.h"
+#include "IRenderPageCollectionEditor.h"
+#include "SlateOptMacros.h"
+
+#define LOCTEXT_NAMESPACE "SRenderPagesPageListTab"
+
+
+BEGIN_SLATE_FUNCTION_BUILD_OPTIMIZATION
+
+void UE::RenderPages::Private::SRenderPagesPageListTab::Construct(const FArguments& InArgs, TSharedPtr<IRenderPageCollectionEditor> InBlueprintEditor)
+{
+	ChildSlot
+	[
+		SNew(SRenderPagesPageList, InBlueprintEditor)
+	];
+}
+
+END_SLATE_FUNCTION_BUILD_OPTIMIZATION
+
+
+#undef LOCTEXT_NAMESPACE
diff --git a/Engine/Plugins/Experimental/RenderPages/Source/RenderPagesEditor/Private/UI/Tabs/SRenderPagesPageListTab.h b/Engine/Plugins/Experimental/RenderPages/Source/RenderPagesEditor/Private/UI/Tabs/SRenderPagesPageListTab.h
new file mode 100644
--- /dev/null
+++ b/Engine/Plugins/Experimental/RenderPages/Source/RenderPagesEditor/Private/UI/Tabs/SRenderPagesPageListTab.h
@@ -0,0 +1,28 @@
+// Copyright Epic Games, Inc. All Rights Reserved.
+
+#pragma once
+
+#include "CoreMinimal.h"
+#include "Widgets/SCompoundWidget.h"
+
+
+namespace UE::RenderPages
+{
+	class IRenderPageCollectionEditor;
+}
+
+
+namespace UE::RenderPages::Private
+{
+	/**
+	 * The page list tab.
+	 */
+	class SRenderPagesPageListTab : public SCompoundWidget
+	{
+	public:
+		SLATE_BEGIN_ARGS(SRenderPagesPageListTab) {}
+		SLATE_END_ARGS()
+
+		void Construct(const FArguments& InArgs, TSharedPtr<IRenderPageCollectionEditor> InBlueprintEditor);
+	};
+}
diff --git a/Engine/Plugins/Experimental/RenderPages/Source/RenderPagesEditor/Private/UI/Tabs/SRenderPagesPagePropertiesTab.cpp b/Engine/Plugins/Experimental/RenderPages/Source/RenderPagesEditor/Private/UI/Tabs/SRenderPagesPagePropertiesTab.cpp
new file mode 100644
--- /dev/null
+++ b/Engine/Plugins/Experimental/RenderPages/Source/RenderPagesEditor/Private/UI/Tabs/SRenderPagesPagePropertiesTab.cpp
@@ -0,0 +1,39 @@
+// Copyright Epic Games, Inc. All Rights Reserved.
+
+#include "UI/Tabs/SRenderPagesPagePropertiesTab.h"
+#include "UI/SRenderPagesPage.h"
+#include "UI/SRenderPagesProps.h"
+#include "IRenderPageCollectionEditor.h"
+#include "SlateOptMacros.h"
+
+#define LOCTEXT_NAMESPACE "SRenderPagesPagePropertiesTab"
+
+
+BEGIN_SLATE_FUNCTION_BUILD_OPTIMIZATION
+
+void UE::RenderPages::Private::SRenderPagesPagePropertiesTab::Construct(const FArguments& InArgs, TSharedPtr<IRenderPageCollectionEditor> InBlueprintEditor)
+{
+	ChildSlot
+	[
+		SNew(SVerticalBox)
+
+		+ SVerticalBox::Slot()
+		.AutoHeight()
+		.Padding(0.0f)
+		[
+			SNew(SRenderPagesPage, InBlueprintEditor)
+		]
+
+		+ SVerticalBox::Slot()
+		.AutoHeight()
+		.Padding(0.0f)
+		[
+			SNew(SRenderPagesProps, InBlueprintEditor)
+		]
+	];
+}
+
+END_SLATE_FUNCTION_BUILD_OPTIMIZATION
+
+
+#undef LOCTEXT_NAMESPACE
diff --git a/Engine/Plugins/Experimental/RenderPages/Source/RenderPagesEditor/Private/UI/Tabs/SRenderPagesPagePropertiesTab.h b/Engine/Plugins/Experimental/RenderPages/Source/RenderPagesEditor/Private/UI/Tabs/SRenderPagesPagePropertiesTab.h
new file mode 100644
--- /dev/null
+++ b/Engine/Plugins/Experimental/RenderPages/Source/RenderPagesEditor/Private/UI/Tabs/SRenderPagesPagePropertiesTab.h
@@ -0,0 +1,28 @@
+// Copyright Epic Games, Inc. All Rights Reserved.
+
+#pragma once
+
+#include "CoreMinimal.h"
+#include "Widgets/SCompoundWidget.h"
+
+
+namespace UE::RenderPages
+{
+	class IRenderPageCollectionEditor;
+}
+
+
+namespace UE::RenderPages::Private
+{
+	/**
+	 * The page properties tab.
+	 */
+	class SRenderPagesPagePropertiesTab : public SCompoundWidget
+	{
+	public:
+		SLATE_BEGIN_ARGS(SRenderPagesPagePropertiesTab) {}
+		SLATE_END_ARGS()
+
+		void Construct(const FArguments& InArgs, TSharedPtr<IRenderPageCollectionEditor> InBlueprintEditor);
+	};
+}
diff --git a/Engine/Plugins/Experimental/RenderPages/Source/RenderPagesEditor/Private/UI/Tabs/SRenderPagesPageViewerTab.cpp b/Engine/Plugins/Experimental/RenderPages/Source/RenderPagesEditor/Private/UI/Tabs/SRenderPagesPageViewerTab.cpp
new file mode 100644
--- /dev/null
+++ b/Engine/Plugins/Experimental/RenderPages/Source/RenderPagesEditor/Private/UI/Tabs/SRenderPagesPageViewerTab.cpp
@@ -0,0 +1,24 @@
+// Copyright Epic Games, Inc. All Rights Reserved.
+
+#include "UI/Tabs/SRenderPagesPageViewerTab.h"
+#include "UI/SRenderPagesPageViewer.h"
+#include "IRenderPageCollectionEditor.h"
+#include "SlateOptMacros.h"
+
+#define LOCTEXT_NAMESPACE "SRenderPagesPageViewerTab"
+
+
+BEGIN_SLATE_FUNCTION_BUILD_OPTIMIZATION
+
+void UE::RenderPages::Private::SRenderPagesPageViewerTab::Construct(const FArguments& InArgs, TSharedPtr<IRenderPageCollectionEditor> InBlueprintEditor)
+{
+	ChildSlot
+	[
+		SNew(SRenderPagesPageViewer, InBlueprintEditor)
+	];
+}
+
+END_SLATE_FUNCTION_BUILD_OPTIMIZATION
+
+
+#undef LOCTEXT_NAMESPACE
diff --git a/Engine/Plugins/Experimental/RenderPages/Source/RenderPagesEditor/Private/UI/Tabs/SRenderPagesPageViewerTab.h b/Engine/Plugins/Experimental/RenderPages/Source/RenderPagesEditor/Private/UI/Tabs/SRenderPagesPageViewerTab.h
new file mode 100644
--- /dev/null
+++ b/Engine/Plugins/Experimental/RenderPages/Source/RenderPagesEditor/Private/UI/Tabs/SRenderPagesPageViewerTab.h
@@ -0,0 +1,28 @@
+// Copyright Epic Games, Inc. All Rights Reserved.
+
+#pragma once
+
+#include "CoreMinimal.h"
+#include "Widgets/SCompoundWidget.h"
+
+
+namespace UE::RenderPages
+{
+	class IRenderPageCollectionEditor;
+}
+
+
+namespace UE::RenderPages::Private
+{
+	/**
+	 * The page viewer tab.
+	 */
+	class SRenderPagesPageViewerTab : public SCompoundWidget
+	{
+	public:
+		SLATE_BEGIN_ARGS(SRenderPagesPageViewerTab) {}
+		SLATE_END_ARGS()
+
+		void Construct(const FArguments& InArgs, TSharedPtr<IRenderPageCollectionEditor> InBlueprintEditor);
+	};
+}
diff --git a/Engine/Plugins/Experimental/RenderPages/Source/RenderPagesEditor/Private/UI/Utils/RenderPagesWidgetUtils.cpp b/Engine/Plugins/Experimental/RenderPages/Source/RenderPagesEditor/Private/UI/Utils/RenderPagesWidgetUtils.cpp
new file mode 100644
--- /dev/null
+++ b/Engine/Plugins/Experimental/RenderPages/Source/RenderPagesEditor/Private/UI/Utils/RenderPagesWidgetUtils.cpp
@@ -0,0 +1,116 @@
+// Copyright Epic Games, Inc. All Rights Reserved.
+
+#include "UI/Utils/RenderPagesWidgetUtils.h"
+#include "IDetailTreeNode.h"
+#include "PropertyHandle.h"
+#include "SlateOptMacros.h"
+#include "Widgets/SBoxPanel.h"
+
+#define LOCTEXT_NAMESPACE "RenderPagesWidgetUtils"
+
+
+BEGIN_SLATE_FUNCTION_BUILD_OPTIMIZATION
+
+TSharedRef<SWidget> UE::RenderPages::Private::RenderPagesWidgetUtils::CreateNodeValueWidget(const TSharedPtr<IDetailTreeNode>& Node)
+{
+	FNodeWidgets NodeWidgets = Node->CreateNodeWidgets();
+
+	TSharedRef<SHorizontalBox> FieldWidget = SNew(SHorizontalBox);
+
+	if (NodeWidgets.ValueWidget)
+	{
+		FieldWidget->AddSlot()
+			.Padding(FMargin(3.0f, 2.0f))
+			.HAlign(HAlign_Right)
+			.FillWidth(1.0f)
+			[
+				NodeWidgets.ValueWidget.ToSharedRef()
+			];
+	}
+	else if (NodeWidgets.WholeRowWidget)
+	{
+		FieldWidget->AddSlot()
+			.Padding(FMargin(3.0f, 2.0f))
+			.FillWidth(1.0f)
+			[
+				NodeWidgets.WholeRowWidget.ToSharedRef()
+			];
+	}
+
+	return FieldWidget;
+}
+
+bool UE::RenderPages::Private::RenderPagesWidgetUtils::FindPropertyHandleRecursive(const TSharedPtr<IPropertyHandle>& PropertyHandle, const FString& PropertyNameOrPath, ERenderPagesFindNodeMethod FindMethod)
+{
+	if (PropertyHandle && PropertyHandle->IsValidHandle())
+	{
+		uint32 ChildrenCount = 0;
+		PropertyHandle->GetNumChildren(ChildrenCount);
+		for (uint32 Index = 0; Index < ChildrenCount; ++Index)
+		{
+			TSharedPtr<IPropertyHandle> ChildHandle = PropertyHandle->GetChildHandle(Index);
+			if (FindPropertyHandleRecursive(ChildHandle, PropertyNameOrPath, FindMethod))
+			{
+				return true;
+			}
+		}
+
+		if (PropertyHandle->GetProperty())
+		{
+			if (FindMethod == ERenderPagesFindNodeMethod::Path)
+			{
+				if (PropertyHandle->GeneratePathToProperty() == PropertyNameOrPath)
+				{
+					return true;
+				}
+			}
+			else if (PropertyHandle->GetProperty()->GetName() == PropertyNameOrPath)
+			{
+				return true;
+			}
+		}
+	}
+
+	return false;
+}
+
+TSharedPtr<IDetailTreeNode> UE::RenderPages::Private::RenderPagesWidgetUtils::FindTreeNodeRecursive(const TSharedRef<IDetailTreeNode>& RootNode, const FString& PropertyNameOrPath, ERenderPagesFindNodeMethod FindMethod)
+{
+	TArray<TSharedRef<IDetailTreeNode>> Children;
+	RootNode->GetChildren(Children);
+	for (TSharedRef<IDetailTreeNode>& Child : Children)
+	{
+		TSharedPtr<IDetailTreeNode> FoundNode = FindTreeNodeRecursive(Child, PropertyNameOrPath, FindMethod);
+		if (FoundNode.IsValid())
+		{
+			return FoundNode;
+		}
+	}
+
+	TSharedPtr<IPropertyHandle> Handle = RootNode->CreatePropertyHandle();
+	if (FindPropertyHandleRecursive(Handle, PropertyNameOrPath, FindMethod))
+	{
+		return RootNode;
+	}
+
+	return nullptr;
+}
+
+TSharedPtr<IDetailTreeNode> UE::RenderPages::Private::RenderPagesWidgetUtils::FindNode(const TArray<TSharedRef<IDetailTreeNode>>& RootNodes, const FString& QualifiedPropertyName, ERenderPagesFindNodeMethod FindMethod)
+{
+	for (const TSharedRef<IDetailTreeNode>& CategoryNode : RootNodes)
+	{
+		TSharedPtr<IDetailTreeNode> FoundNode = FindTreeNodeRecursive(CategoryNode, QualifiedPropertyName, FindMethod);
+		if (FoundNode.IsValid())
+		{
+			return FoundNode;
+		}
+	}
+
+	return nullptr;
+}
+
+END_SLATE_FUNCTION_BUILD_OPTIMIZATION
+
+
+#undef LOCTEXT_NAMESPACE
diff --git a/Engine/Plugins/Experimental/RenderPages/Source/RenderPagesEditor/Private/UI/Utils/RenderPagesWidgetUtils.h b/Engine/Plugins/Experimental/RenderPages/Source/RenderPagesEditor/Private/UI/Utils/RenderPagesWidgetUtils.h
new file mode 100644
--- /dev/null
+++ b/Engine/Plugins/Experimental/RenderPages/Source/RenderPagesEditor/Private/UI/Utils/RenderPagesWidgetUtils.h
@@ -0,0 +1,33 @@
+// Copyright Epic Games, Inc. All Rights Reserved.
+
+#pragma once
+
+#include "CoreMinimal.h"
+
+
+class IDetailTreeNode;
+class IPropertyHandle;
+class SWidget;
+
+
+namespace UE::RenderPages::Private
+{
+	/**
+	 * A static class containing utility functions for widgets.
+	 */
+	class RenderPagesWidgetUtils final
+	{
+	public:
+		enum class ERenderPagesFindNodeMethod : uint8
+		{
+			Name, // Find a node by providing its property name
+			Path // Find a node by providing a property path
+		};
+
+	public:
+		static TSharedRef<SWidget> CreateNodeValueWidget(const TSharedPtr<IDetailTreeNode>& Node);
+		static bool FindPropertyHandleRecursive(const TSharedPtr<IPropertyHandle>& PropertyHandle, const FString& PropertyNameOrPath, ERenderPagesFindNodeMethod FindMethod);
+		static TSharedPtr<IDetailTreeNode> FindTreeNodeRecursive(const TSharedRef<IDetailTreeNode>& RootNode, const FString& PropertyNameOrPath, ERenderPagesFindNodeMethod FindMethod);
+		static TSharedPtr<IDetailTreeNode> FindNode(const TArray<TSharedRef<IDetailTreeNode>>& RootNodes, const FString& QualifiedPropertyName, ERenderPagesFindNodeMethod FindMethod);
+	};
+}
diff --git a/Engine/Plugins/Experimental/RenderPages/Source/RenderPagesEditor/Public/Factories/IRenderPagePropsSourceWidgetFactory.h b/Engine/Plugins/Experimental/RenderPages/Source/RenderPagesEditor/Public/Factories/IRenderPagePropsSourceWidgetFactory.h
new file mode 100644
--- /dev/null
+++ b/Engine/Plugins/Experimental/RenderPages/Source/RenderPagesEditor/Public/Factories/IRenderPagePropsSourceWidgetFactory.h
@@ -0,0 +1,29 @@
+// Copyright Epic Games, Inc. All Rights Reserved.
+
+#pragma once
+
+#include "CoreMinimal.h"
+#include "IRenderPageCollectionEditor.h"
+
+
+class URenderPagePropsSourceBase;
+
+namespace UE::RenderPages::Private
+{
+	class SRenderPagesPropsBase;
+}
+
+
+namespace UE::RenderPages
+{
+	/**
+	 * The props source widget factory interface.
+	 * Implementations should create a widget for the given props source.
+	 */
+	class RENDERPAGESEDITOR_API IRenderPagePropsSourceWidgetFactory
+	{
+	public:
+		virtual ~IRenderPagePropsSourceWidgetFactory() = default;
+		virtual TSharedPtr<Private::SRenderPagesPropsBase> CreateInstance(URenderPagePropsSourceBase* PropsSource, TSharedPtr<IRenderPageCollectionEditor> BlueprintEditor) { return nullptr; }
+	};
+}
diff --git a/Engine/Plugins/Experimental/RenderPages/Source/RenderPagesEditor/Public/IRenderPageCollectionEditor.h b/Engine/Plugins/Experimental/RenderPages/Source/RenderPagesEditor/Public/IRenderPageCollectionEditor.h
new file mode 100644
--- /dev/null
+++ b/Engine/Plugins/Experimental/RenderPages/Source/RenderPagesEditor/Public/IRenderPageCollectionEditor.h
@@ -0,0 +1,98 @@
+// Copyright Epic Games, Inc. All Rights Reserved.
+
+#pragma once
+
+#include "CoreMinimal.h"
+#include "BlueprintEditor.h"
+
+
+class URenderPagesMoviePipelineRenderJob;
+class URenderPage;
+class URenderPageCollection;
+class URenderPagesBlueprint;
+
+namespace UE::RenderPages::Private
+{
+	class FRenderPagesBlueprintEditorToolbar;
+}
+
+
+namespace UE::RenderPages
+{
+	/**
+	 * The render page editor interface.
+	 */
+	class IRenderPageCollectionEditor : public FBlueprintEditor
+	{
+	public:
+		/** @return The render pages blueprint currently being edited in this editor. */
+		virtual class URenderPagesBlueprint* GetRenderPagesBlueprint() const = 0;
+
+		/** @return The render pages collection instance. */
+		virtual URenderPageCollection* GetInstance() const = 0;
+
+		/** @return The render pages toolbar. */
+		virtual TSharedPtr<Private::FRenderPagesBlueprintEditorToolbar> GetRenderPagesToolbarBuilder() = 0;
+
+		/** Returns whether it is currently rendering or playing (so changes in the level and such should be ignored). */
+		virtual bool IsCurrentlyRenderingOrPlaying() const { return IsBatchRendering() || IsPreviewRendering() || IsValid(GEditor->PlayWorld); }
+
+		/** Returns whether it can currently render (like a preview render or a batch render) or not. */
+		virtual bool CanCurrentlyRender() const { return !IsCurrentlyRenderingOrPlaying(); }
+
+		/** Returns whether it is currently batch rendering or not. */
+		virtual bool IsBatchRendering() const = 0;
+
+		/** Returns the current batch rendering job, or null if it's not currently batch rendering. */
+		virtual URenderPagesMoviePipelineRenderJob* GetBatchRenderJob() const = 0;
+
+		/** Returns whether it is currently preview rendering or not. */
+		virtual bool IsPreviewRendering() const = 0;
+
+		/** Returns the current preview rendering job, or null if it's not currently rendering a preview. */
+		virtual URenderPagesMoviePipelineRenderJob* GetPreviewRenderJob() const = 0;
+
+		/** Sets the current preview rendering job, set it to null if it's not currently rendering a preview. */
+		virtual void SetPreviewRenderJob(URenderPagesMoviePipelineRenderJob* Job) = 0;
+
+		/** Marks the editing asset as modified. */
+		virtual void MarkAsModified() = 0;
+
+		/** Get the currently selected render pages. */
+		virtual TArray<URenderPage*> GetSelectedRenderPages() const = 0;
+
+		/** Set the selected render pages. */
+		virtual void SetSelectedRenderPages(const TArray<URenderPage*>& RenderPages) = 0;
+
+		DECLARE_MULTICAST_DELEGATE(FOnRenderPagesChanged);
+		virtual FOnRenderPagesChanged& OnRenderPagesChanged() { return OnRenderPagesChangedDelegate; }
+
+		DECLARE_MULTICAST_DELEGATE_OneParam(FOnRenderPageCreated, URenderPage*);
+		virtual FOnRenderPageCreated& OnRenderPageCreated() { return OnRenderPageCreatedDelegate; }
+
+		DECLARE_MULTICAST_DELEGATE(FOnRenderPagesSelectionChanged);
+		virtual FOnRenderPagesSelectionChanged& OnRenderPagesSelectionChanged() { return OnRenderPagesSelectionChangedDelegate; }
+
+		DECLARE_MULTICAST_DELEGATE_OneParam(FOnRenderPagesBatchRenderingStarted, URenderPagesMoviePipelineRenderJob*);
+		virtual FOnRenderPagesBatchRenderingStarted& OnRenderPagesBatchRenderingStarted() { return OnRenderPagesBatchRenderingStartedDelegate; }
+
+		DECLARE_MULTICAST_DELEGATE_OneParam(FOnRenderPagesBatchRenderingFinished, URenderPagesMoviePipelineRenderJob*);
+		virtual FOnRenderPagesBatchRenderingFinished& OnRenderPagesBatchRenderingFinished() { return OnRenderPagesBatchRenderingFinishedDelegate; }
+
+	protected:
+		/** The delegate for when render pages in the collection changed. */
+		FOnRenderPagesChanged OnRenderPagesChangedDelegate;
+
+		/** The delegate for when a render page is created. */
+		FOnRenderPageCreated OnRenderPageCreatedDelegate;
+
+		/** The delegate for when the selection of render pages changed. */
+		FOnRenderPagesSelectionChanged OnRenderPagesSelectionChangedDelegate;
+
+		/** The delegate for when batch rendering of pages started. */
+		FOnRenderPagesBatchRenderingStarted OnRenderPagesBatchRenderingStartedDelegate;
+
+		/** The delegate for when batch rendering of pages ended, successful or not. */
+		FOnRenderPagesBatchRenderingFinished OnRenderPagesBatchRenderingFinishedDelegate;
+	};
+}
