First usable version of the render pages plugin.

diff --git a/Engine/Plugins/Experimental/RenderPages/Source/RenderPagesEditor/Public/IRenderPagesEditorModule.h b/Engine/Plugins/Experimental/RenderPages/Source/RenderPagesEditor/Public/IRenderPagesEditorModule.h
new file mode 100644
--- /dev/null
+++ b/Engine/Plugins/Experimental/RenderPages/Source/RenderPagesEditor/Public/IRenderPagesEditorModule.h
@@ -0,0 +1,53 @@
+// Copyright Epic Games, Inc. All Rights Reserved.
+
+#pragma once
+
+#include "CoreMinimal.h"
+#include "Toolkits/AssetEditorToolkit.h"
+#include "IRenderPageCollectionEditor.h"
+#include "Modules/ModuleManager.h"
+
+
+class URenderPagesBlueprint;
+class URenderPagePropsSourceBase;
+enum class ERenderPagePropsSourceType : uint8;
+
+namespace UE::RenderPages
+{
+	class IRenderPagePropsSourceWidgetFactory;
+}
+
+namespace UE::RenderPages::Private
+{
+	class SRenderPagesPropsBase;
+}
+
+
+namespace UE::RenderPages
+{
+	/**
+	 * RenderPagesEditor module interface.
+	 */
+	class IRenderPagesEditorModule : public IModuleInterface, public IHasMenuExtensibility, public IHasToolBarExtensibility
+	{
+	public:
+		/**
+		 * Singleton-like access to IRenderPagesEditorModule.
+		 *
+		 * @return Returns RenderPagesEditorModule singleton instance, loading the module on demand if needed.
+		 */
+		static FORCEINLINE IRenderPagesEditorModule& Get()
+		{
+			return FModuleManager::LoadModuleChecked<IRenderPagesEditorModule>(TEXT("RenderPagesEditor"));
+		}
+
+		/** Creates an instance of the render page collection editor. */
+		virtual TSharedRef<IRenderPageCollectionEditor> CreateRenderPageCollectionEditor(const EToolkitMode::Type Mode, const TSharedPtr<IToolkitHost>& InitToolkitHost, URenderPagesBlueprint* InBlueprint) = 0;
+
+		/** Creates a props source widget for the given props source. */
+		virtual TSharedPtr<Private::SRenderPagesPropsBase> CreatePropsSourceWidget(URenderPagePropsSourceBase* PropsSource, TSharedPtr<IRenderPageCollectionEditor> BlueprintEditor) = 0;
+
+		/** Returns all the factories for creating widgets for props sources. */
+		virtual const TMap<ERenderPagePropsSourceType, TSharedPtr<IRenderPagePropsSourceWidgetFactory>>& GetPropsSourceWidgetFactories() const = 0;
+	};
+}
diff --git a/Engine/Plugins/Experimental/RenderPages/Source/RenderPagesEditor/Public/RenderPagesEditorLog.h b/Engine/Plugins/Experimental/RenderPages/Source/RenderPagesEditor/Public/RenderPagesEditorLog.h
new file mode 100644
--- /dev/null
+++ b/Engine/Plugins/Experimental/RenderPages/Source/RenderPagesEditor/Public/RenderPagesEditorLog.h
@@ -0,0 +1,8 @@
+// Copyright Epic Games, Inc. All Rights Reserved.
+
+#pragma once
+
+#include "CoreMinimal.h"
+
+
+DECLARE_LOG_CATEGORY_CLASS(LogRenderPagesEditor, Log, All);
diff --git a/Engine/Plugins/Experimental/RenderPages/Source/RenderPagesEditor/RenderPagesEditor.Build.cs b/Engine/Plugins/Experimental/RenderPages/Source/RenderPagesEditor/RenderPagesEditor.Build.cs
new file mode 100644
--- /dev/null
+++ b/Engine/Plugins/Experimental/RenderPages/Source/RenderPagesEditor/RenderPagesEditor.Build.cs
@@ -0,0 +1,45 @@
+// Copyright Epic Games, Inc. All Rights Reserved.
+
+using UnrealBuildTool;
+
+public class RenderPagesEditor : ModuleRules
+{
+	public RenderPagesEditor(ReadOnlyTargetRules Target) : base(Target)
+	{
+		PCHUsage = PCHUsageMode.UseExplicitOrSharedPCHs;
+
+		PublicDependencyModuleNames.AddRange(
+			new string[]
+			{
+				"Core",
+				"CoreUObject",
+				"Engine",
+				"RemoteControl",
+			}
+		);
+
+
+		PrivateDependencyModuleNames.AddRange(
+			new string[]
+			{
+				"AssetTools",
+				"DesktopPlatform",
+				"EditorStyle",
+				"GraphEditor",
+				"InputCore",
+				"Kismet",
+				"LevelSequence",
+				"MovieRenderPipelineCore",
+				"MovieScene",
+				"Projects",
+				"PropertyEditor",
+				"RenderPages",
+				"RenderPagesDeveloper",
+				"Slate",
+				"SlateCore",
+				"ToolMenus",
+				"UnrealEd",
+			}
+		);
+	}
+}
\ No newline at end of file
