Fixed a text conflict from LOCTEXT macro for namespace "SRenderPagesPageList" and key "PagesCollectionIDColumnHeader".

diff --git a/Engine/Plugins/Experimental/RenderPages/Source/RenderPagesEditor/Private/UI/SRenderPagesPageList.cpp b/Engine/Plugins/Experimental/RenderPages/Source/RenderPagesEditor/Private/UI/SRenderPagesPageList.cpp
--- a/Engine/Plugins/Experimental/RenderPages/Source/RenderPagesEditor/Private/UI/SRenderPagesPageList.cpp
+++ b/Engine/Plugins/Experimental/RenderPages/Source/RenderPagesEditor/Private/UI/SRenderPagesPageList.cpp
@@ -124,7 +124,7 @@ void UE::RenderPages::Private::SRenderPagesPageList::Construct(const FArguments&
 					.FixedWidth(36.0f)
 
 					+ SHeaderRow::Column(FRenderPagesCollectionColumns::IsEnabled)
-					.DefaultLabel(LOCTEXT("PagesCollectionIDColumnHeader", "Enabled"))
+					.DefaultLabel(LOCTEXT("PagesCollectionIsEnabledColumnHeader", "Enabled"))
 					.FixedWidth(30.0f) //55.0f for text : "Enabled"
 					[
 						SAssignNew(RenderPageEnabledHeaderCheckbox, SCheckBox)
