Fixed ue-build-health issues.

diff --git a/Engine/Plugins/Experimental/RenderPages/Source/RenderPages/Private/RenderPage/RenderPageCollection.cpp b/Engine/Plugins/Experimental/RenderPages/Source/RenderPages/Private/RenderPage/RenderPageCollection.cpp
--- a/Engine/Plugins/Experimental/RenderPages/Source/RenderPages/Private/RenderPage/RenderPageCollection.cpp
+++ b/Engine/Plugins/Experimental/RenderPages/Source/RenderPages/Private/RenderPage/RenderPageCollection.cpp
@@ -86,7 +86,7 @@ bool URenderPage::SetSequenceStartFrame(const int32 NewCustomStartFrame)
 	while (SequenceStartFrame.IsSet() && (SequenceStartFrame.Get(0) > NewCustomStartFrame))
 	{
 		StartFrame--;
-		if (StartFrame <= -2147483648)
+		if (StartFrame <= INT32_MIN)
 		{
 			return false;
 		}
@@ -96,7 +96,7 @@ bool URenderPage::SetSequenceStartFrame(const int32 NewCustomStartFrame)
 	while (SequenceStartFrame.IsSet() && (SequenceStartFrame.Get(0) < NewCustomStartFrame))
 	{
 		StartFrame++;
-		if (StartFrame >= 2147483647)
+		if (StartFrame >= INT32_MAX)
 		{
 			return false;
 		}
@@ -116,7 +116,7 @@ bool URenderPage::SetSequenceEndFrame(const int32 NewCustomStartFrame)
 	while (SequenceEndFrame.IsSet() && (SequenceEndFrame.Get(0) > NewCustomStartFrame))
 	{
 		EndFrame--;
-		if (EndFrame <= -2147483648)
+		if (EndFrame <= INT32_MIN)
 		{
 			return false;
 		}
@@ -126,7 +126,7 @@ bool URenderPage::SetSequenceEndFrame(const int32 NewCustomStartFrame)
 	while (SequenceEndFrame.IsSet() && (SequenceEndFrame.Get(0) < NewCustomStartFrame))
 	{
 		EndFrame++;
-		if (EndFrame >= 2147483647)
+		if (EndFrame >= INT32_MAX)
 		{
 			return false;
 		}
diff --git a/Engine/Plugins/Experimental/RenderPages/Source/RenderPages/Private/RenderPage/RenderPageMoviePipelineJob.cpp b/Engine/Plugins/Experimental/RenderPages/Source/RenderPages/Private/RenderPage/RenderPageMoviePipelineJob.cpp
--- a/Engine/Plugins/Experimental/RenderPages/Source/RenderPages/Private/RenderPage/RenderPageMoviePipelineJob.cpp
+++ b/Engine/Plugins/Experimental/RenderPages/Source/RenderPages/Private/RenderPage/RenderPageMoviePipelineJob.cpp
@@ -43,7 +43,7 @@ URenderPagesMoviePipelineRenderJob* URenderPagesMoviePipelineRenderJob::Create(c
 	{
 		if (UMoviePipelinePIEExecutor* ActiveExecutorPIE = Cast<UMoviePipelinePIEExecutor>(RenderJob->ActiveExecutor))
 		{
-			//TODO: when it becomes available on ue5-main, uncomment:  ActiveExecutorPIE->SetIsRenderingOffscreen(true);
+			ActiveExecutorPIE->SetIsRenderingOffscreen(true);
 		}
 	}
 
@@ -114,7 +114,10 @@ URenderPagesMoviePipelineRenderJob* URenderPagesMoviePipelineRenderJob::Create(c
 			}
 			if (Args.bForceOutputImage && !bFound)
 			{
-				NewJob->GetConfiguration()->FindOrAddSettingByClass(UMoviePipelineImageSequenceOutput_PNG::StaticClass());
+				if(UMoviePipelineImageSequenceOutput_PNG* NewSetting = Cast<UMoviePipelineImageSequenceOutput_PNG>(NewJob->GetConfiguration()->FindOrAddSettingByClass(UMoviePipelineImageSequenceOutput_PNG::StaticClass())))
+				{
+					NewSetting->bWriteAlpha = false;
+				}
 			}
 		}
 
diff --git a/Engine/Plugins/Experimental/RenderPages/Source/RenderPages/RenderPages.Build.cs b/Engine/Plugins/Experimental/RenderPages/Source/RenderPages/RenderPages.Build.cs
--- a/Engine/Plugins/Experimental/RenderPages/Source/RenderPages/RenderPages.Build.cs
+++ b/Engine/Plugins/Experimental/RenderPages/Source/RenderPages/RenderPages.Build.cs
@@ -27,6 +27,7 @@ public class RenderPages : ModuleRules
 				"MovieRenderPipelineCore",
 				"MovieRenderPipelineEditor",
 				"MovieRenderPipelineRenderPasses",
+				"MovieScene",
 				"Serialization",
 			}
 		);
