Fix non unity errors and #pragma once in cpp warnings
#rnx

diff --git a/Engine/Plugins/Experimental/MotorSimOutputMotoSynth/Source/MotorSimOutputMotoSynth/Private/MotorSimOutputMotoSynth.cpp b/Engine/Plugins/Experimental/MotorSimOutputMotoSynth/Source/MotorSimOutputMotoSynth/Private/MotorSimOutputMotoSynth.cpp
--- a/Engine/Plugins/Experimental/MotorSimOutputMotoSynth/Source/MotorSimOutputMotoSynth/Private/MotorSimOutputMotoSynth.cpp
+++ b/Engine/Plugins/Experimental/MotorSimOutputMotoSynth/Source/MotorSimOutputMotoSynth/Private/MotorSimOutputMotoSynth.cpp
@@ -1,7 +1,5 @@
 // Copyright Epic Games, Inc. All Rights Reserved.
 
-#pragma once
-
 #include "MotorSimOutputMotoSynth.h"
 #include "AudioMotorSimTypes.h"
 
diff --git a/Engine/Plugins/Experimental/RenderPages/Source/RenderPages/Private/Factories/RenderPagePropsSourceFactoryLocal.cpp b/Engine/Plugins/Experimental/RenderPages/Source/RenderPages/Private/Factories/RenderPagePropsSourceFactoryLocal.cpp
--- a/Engine/Plugins/Experimental/RenderPages/Source/RenderPages/Private/Factories/RenderPagePropsSourceFactoryLocal.cpp
+++ b/Engine/Plugins/Experimental/RenderPages/Source/RenderPages/Private/Factories/RenderPagePropsSourceFactoryLocal.cpp
@@ -1,7 +1,5 @@
 // Copyright Epic Games, Inc. All Rights Reserved.
 
-#pragma once
-
 #include "Factories/RenderPagePropsSourceFactoryLocal.h"
 #include "RenderPage/RenderPagePropsSource.h"
 
diff --git a/Engine/Plugins/Experimental/RenderPages/Source/RenderPages/Private/Factories/RenderPagePropsSourceFactoryRemoteControl.cpp b/Engine/Plugins/Experimental/RenderPages/Source/RenderPages/Private/Factories/RenderPagePropsSourceFactoryRemoteControl.cpp
--- a/Engine/Plugins/Experimental/RenderPages/Source/RenderPages/Private/Factories/RenderPagePropsSourceFactoryRemoteControl.cpp
+++ b/Engine/Plugins/Experimental/RenderPages/Source/RenderPages/Private/Factories/RenderPagePropsSourceFactoryRemoteControl.cpp
@@ -1,7 +1,5 @@
 // Copyright Epic Games, Inc. All Rights Reserved.
 
-#pragma once
-
 #include "Factories/RenderPagePropsSourceFactoryRemoteControl.h"
 #include "RenderPage/RenderPagePropsSource.h"
 
diff --git a/Engine/Plugins/Experimental/RenderPages/Source/RenderPagesEditor/Private/Factories/RenderPagePropsSourceWidgetFactoryRemoteControl.cpp b/Engine/Plugins/Experimental/RenderPages/Source/RenderPagesEditor/Private/Factories/RenderPagePropsSourceWidgetFactoryRemoteControl.cpp
--- a/Engine/Plugins/Experimental/RenderPages/Source/RenderPagesEditor/Private/Factories/RenderPagePropsSourceWidgetFactoryRemoteControl.cpp
+++ b/Engine/Plugins/Experimental/RenderPages/Source/RenderPagesEditor/Private/Factories/RenderPagePropsSourceWidgetFactoryRemoteControl.cpp
@@ -1,7 +1,5 @@
 // Copyright Epic Games, Inc. All Rights Reserved.
 
-#pragma once
-
 #include "Factories/RenderPagePropsSourceWidgetFactoryRemoteControl.h"
 #include "UI/SRenderPagesPropsRemoteControl.h"
 #include "RenderPage/RenderPagePropsSource.h"
diff --git a/Engine/Source/Runtime/Core/Private/ProfilingDebugging/TraceAuxiliary.cpp b/Engine/Source/Runtime/Core/Private/ProfilingDebugging/TraceAuxiliary.cpp
--- a/Engine/Source/Runtime/Core/Private/ProfilingDebugging/TraceAuxiliary.cpp
+++ b/Engine/Source/Runtime/Core/Private/ProfilingDebugging/TraceAuxiliary.cpp
@@ -1,6 +1,7 @@
 // Copyright Epic Games, Inc. All Rights Reserved.
 
 #include "ProfilingDebugging/TraceAuxiliary.h"
+#include "ProfilingDebugging/StringsTrace.h"
 #include "Trace/Trace.h"
 #include "CoreGlobals.h"
 
diff --git a/Engine/Source/Runtime/Core/Public/Misc/QueuedThreadPool.h b/Engine/Source/Runtime/Core/Public/Misc/QueuedThreadPool.h
--- a/Engine/Source/Runtime/Core/Public/Misc/QueuedThreadPool.h
+++ b/Engine/Source/Runtime/Core/Public/Misc/QueuedThreadPool.h
@@ -3,7 +3,9 @@
 #pragma once
 
 #include "CoreTypes.h"
+#include "Containers/Array.h"
 #include "GenericPlatform/GenericPlatformAffinity.h"
+#include "Templates/Function.h"
 
 class IQueuedWork;
 
