Remove #pragma once from cpp
#rnx

diff --git a/Engine/Plugins/Experimental/RenderPages/Source/RenderPagesEditor/Private/Factories/RenderPagePropsSourceWidgetFactoryLocal.cpp b/Engine/Plugins/Experimental/RenderPages/Source/RenderPagesEditor/Private/Factories/RenderPagePropsSourceWidgetFactoryLocal.cpp
--- a/Engine/Plugins/Experimental/RenderPages/Source/RenderPagesEditor/Private/Factories/RenderPagePropsSourceWidgetFactoryLocal.cpp
+++ b/Engine/Plugins/Experimental/RenderPages/Source/RenderPagesEditor/Private/Factories/RenderPagePropsSourceWidgetFactoryLocal.cpp
@@ -1,7 +1,5 @@
 // Copyright Epic Games, Inc. All Rights Reserved.
 
-#pragma once
-
 #include "Factories/RenderPagePropsSourceWidgetFactoryLocal.h"
 #include "UI/SRenderPagesPropsLocal.h"
 #include "RenderPage/RenderPagePropsSource.h"
