Render Pages plugin:
- Fixed issues with the page selection (was reset on compile, didn't take search filtering into account, etc).
- Fixed the page viewer tab buttons (style).

diff --git a/Engine/Plugins/Experimental/RenderPages/Source/RenderPagesEditor/Private/Styles/RenderPagesEditorStyle.cpp b/Engine/Plugins/Experimental/RenderPages/Source/RenderPagesEditor/Private/Styles/RenderPagesEditorStyle.cpp
--- a/Engine/Plugins/Experimental/RenderPages/Source/RenderPagesEditor/Private/Styles/RenderPagesEditorStyle.cpp
+++ b/Engine/Plugins/Experimental/RenderPages/Source/RenderPagesEditor/Private/Styles/RenderPagesEditorStyle.cpp
@@ -8,7 +8,6 @@
 #include "Interfaces/IPluginManager.h"
 #include "Styling/SlateStyle.h"
 #include "Styling/SlateStyleRegistry.h"
-#include "Styling/StyleColors.h"
 
 
 TSharedPtr<FSlateStyleSet> UE::RenderPages::Private::FRenderPagesEditorStyle::StyleInstance = nullptr;
@@ -96,13 +95,6 @@ TSharedRef<FSlateStyleSet> UE::RenderPages::Private::FRenderPagesEditorStyle::Cr
 		.SetPressedPadding(FMargin())
 	);
 
-	Style->Set("TabButton",
-		FButtonStyle(FAppStyle::Get().GetWidgetStyle<FButtonStyle>("Button"))
-		.SetDisabled(BOX_BRUSH("Common/FlatButton", 2.0f / 8.0f, FAppStyle::GetSlateColor("SelectionColor")))
-		.SetDisabledForeground(FStyleColors::White)
-	);
-
-
 	return Style;
 }
 
diff --git a/Engine/Plugins/Experimental/RenderPages/Source/RenderPagesEditor/Private/Toolkit/RenderPageCollectionEditor.cpp b/Engine/Plugins/Experimental/RenderPages/Source/RenderPagesEditor/Private/Toolkit/RenderPageCollectionEditor.cpp
--- a/Engine/Plugins/Experimental/RenderPages/Source/RenderPagesEditor/Private/Toolkit/RenderPageCollectionEditor.cpp
+++ b/Engine/Plugins/Experimental/RenderPages/Source/RenderPagesEditor/Private/Toolkit/RenderPageCollectionEditor.cpp
@@ -269,14 +269,6 @@ bool UE::RenderPages::Private::FRenderPageCollectionEditor::IsCompilingEnabled()
 	return true;
 }
 
-void UE::RenderPages::Private::FRenderPageCollectionEditor::NewDocument_OnClicked(ECreatedDocumentType GraphType)
-{
-	if (GraphType != CGT_NewFunctionGraph)
-	{
-		return;
-	}
-}
-
 bool UE::RenderPages::Private::FRenderPageCollectionEditor::IsSectionVisible(NodeSectionID::Type InSectionID) const
 {
 	switch (InSectionID)
@@ -358,22 +350,6 @@ bool UE::RenderPages::Private::FRenderPageCollectionEditor::CanAddNewLocalVariab
 	return (Cast<URenderPagesGraph>(GetFocusedGraph()) != nullptr);
 }
 
-void UE::RenderPages::Private::FRenderPageCollectionEditor::OnAddNewLocalVariable()
-{
-	if (!CanAddNewLocalVariable())
-	{
-		return;
-	}
-}
-
-void UE::RenderPages::Private::FRenderPageCollectionEditor::OnPasteNewLocalVariable(const FBPVariableDescription& VariableDescription)
-{
-	if (!CanAddNewLocalVariable())
-	{
-		return;
-	}
-}
-
 void UE::RenderPages::Private::FRenderPageCollectionEditor::Compile()
 {
 	DestroyInstance();
@@ -431,6 +407,9 @@ TArray<URenderPage*> UE::RenderPages::Private::FRenderPageCollectionEditor::GetS
 
 void UE::RenderPages::Private::FRenderPageCollectionEditor::SetSelectedRenderPages(const TArray<URenderPage*>& RenderPages)
 {
+	TSet<FGuid> PreviouslySelectedRenderPagesIds;
+	PreviouslySelectedRenderPagesIds.Append(SelectedRenderPagesIds);
+
 	SelectedRenderPagesIds.Empty();
 	for (URenderPage* Page : RenderPages)
 	{
@@ -441,7 +420,19 @@ void UE::RenderPages::Private::FRenderPageCollectionEditor::SetSelectedRenderPag
 		SelectedRenderPagesIds.Add(Page->GetId());
 	}
 
-	OnRenderPagesSelectionChanged().Broadcast();
+	if (SelectedRenderPagesIds.Num() != PreviouslySelectedRenderPagesIds.Num())
+	{
+		OnRenderPagesSelectionChanged().Broadcast();
+		return;
+	}
+	for (const FGuid& PageId : SelectedRenderPagesIds)
+	{
+		if (!PreviouslySelectedRenderPagesIds.Contains(PageId))
+		{
+			OnRenderPagesSelectionChanged().Broadcast();
+			return;
+		}
+	}
 }
 
 FName UE::RenderPages::Private::FRenderPageCollectionEditor::GetToolkitFName() const
@@ -479,10 +470,6 @@ void UE::RenderPages::Private::FRenderPageCollectionEditor::InitToolMenuContext(
 	FBlueprintEditor::InitToolMenuContext(MenuContext);
 }
 
-void UE::RenderPages::Private::FRenderPageCollectionEditor::OnToolkitHostingStarted(const TSharedRef<IToolkit>& Toolkit) {}
-
-void UE::RenderPages::Private::FRenderPageCollectionEditor::OnToolkitHostingFinished(const TSharedRef<IToolkit>& Toolkit) {}
-
 void UE::RenderPages::Private::FRenderPageCollectionEditor::Tick(float DeltaTime)
 {
 	FBlueprintEditor::Tick(DeltaTime);
diff --git a/Engine/Plugins/Experimental/RenderPages/Source/RenderPagesEditor/Private/Toolkit/RenderPageCollectionEditor.h b/Engine/Plugins/Experimental/RenderPages/Source/RenderPagesEditor/Private/Toolkit/RenderPageCollectionEditor.h
--- a/Engine/Plugins/Experimental/RenderPages/Source/RenderPagesEditor/Private/Toolkit/RenderPageCollectionEditor.h
+++ b/Engine/Plugins/Experimental/RenderPages/Source/RenderPagesEditor/Private/Toolkit/RenderPageCollectionEditor.h
@@ -46,7 +46,6 @@ namespace UE::RenderPages::Private
 		virtual bool IsEditable(UEdGraph* InGraph) const override;
 		virtual bool IsCompilingEnabled() const override;
 		virtual bool IsInAScriptingMode() const override { return true; }
-		virtual void NewDocument_OnClicked(ECreatedDocumentType GraphType) override;
 		virtual bool IsSectionVisible(NodeSectionID::Type InSectionID) const override;
 		virtual FText GetGraphDecorationString(UEdGraph* InGraph) const override;
 		virtual void OnActiveTabChanged(TSharedPtr<SDockTab> PreviouslyActive, TSharedPtr<SDockTab> NewlyActivated) override;
@@ -57,8 +56,6 @@ namespace UE::RenderPages::Private
 		virtual bool ShouldLoadBPLibrariesFromAssetRegistry() override { return false; }// Render Pages BP does not use regular BP function libraries, no need to load them.
 
 		virtual bool CanAddNewLocalVariable() const override;
-		virtual void OnAddNewLocalVariable() override;
-		virtual void OnPasteNewLocalVariable(const FBPVariableDescription& VariableDescription) override;
 
 		virtual void Compile() override;
 		//~ End FBlueprintEditor Interface
@@ -86,11 +83,6 @@ namespace UE::RenderPages::Private
 		virtual void InitToolMenuContext(FToolMenuContext& MenuContext) override;
 		//~ End IToolkit Interface
 
-		//~ Begin IToolkitHost Interface
-		virtual void OnToolkitHostingStarted(const TSharedRef<IToolkit>& Toolkit) override;
-		virtual void OnToolkitHostingFinished(const TSharedRef<IToolkit>& Toolkit) override;
-		//~ End IToolkitHost Interface
-
 		//~ Begin FTickableEditorObject Interface
 		virtual void Tick(float DeltaTime) override;
 		virtual TStatId GetStatId() const override;
diff --git a/Engine/Plugins/Experimental/RenderPages/Source/RenderPagesEditor/Private/UI/SRenderPagesPageList.cpp b/Engine/Plugins/Experimental/RenderPages/Source/RenderPagesEditor/Private/UI/SRenderPagesPageList.cpp
--- a/Engine/Plugins/Experimental/RenderPages/Source/RenderPagesEditor/Private/UI/SRenderPagesPageList.cpp
+++ b/Engine/Plugins/Experimental/RenderPages/Source/RenderPagesEditor/Private/UI/SRenderPagesPageList.cpp
@@ -282,6 +282,14 @@ void UE::RenderPages::Private::SRenderPagesPageList::Refresh()
 		if (RenderPageListWidget.IsValid())
 		{
 			RenderPageListWidget->RebuildList(); // rebuild is needed (instead of using RequestListRefresh()), because otherwise it won't show the changes made to the URenderPage variables
+
+			TArray<URenderPage*> SelectedItems = BlueprintEditor->GetSelectedRenderPages().FilterByPredicate([this](URenderPage* Item)
+			{
+				return IsValid(Item) && RenderPages.Contains(Item);
+			});
+			RenderPageListWidget->ClearSelection();
+			RenderPageListWidget->SetItemSelection(SelectedItems, true);
+			BlueprintEditor->SetSelectedRenderPages(SelectedItems);
 		}
 	}
 }
@@ -305,6 +313,10 @@ TSharedRef<ITableRow> UE::RenderPages::Private::SRenderPagesPageList::HandlePage
 
 void UE::RenderPages::Private::SRenderPagesPageList::HandlePagesCollectionSelectionChanged(URenderPage* Item, ESelectInfo::Type SelectInfo)
 {
+	if (SelectInfo == ESelectInfo::Type::Direct)
+	{
+		return;
+	}
 	if (const TSharedPtr<IRenderPageCollectionEditor> BlueprintEditor = BlueprintEditorWeakPtr.Pin())
 	{
 		BlueprintEditor->SetSelectedRenderPages(RenderPageListWidget->GetSelectedItems());
diff --git a/Engine/Plugins/Experimental/RenderPages/Source/RenderPagesEditor/Private/UI/SRenderPagesPageViewer.cpp b/Engine/Plugins/Experimental/RenderPages/Source/RenderPagesEditor/Private/UI/SRenderPagesPageViewer.cpp
--- a/Engine/Plugins/Experimental/RenderPages/Source/RenderPagesEditor/Private/UI/SRenderPagesPageViewer.cpp
+++ b/Engine/Plugins/Experimental/RenderPages/Source/RenderPagesEditor/Private/UI/SRenderPagesPageViewer.cpp
@@ -6,8 +6,7 @@
 #include "UI/SRenderPagesPageViewerRendered.h"
 #include "IRenderPageCollectionEditor.h"
 #include "SlateOptMacros.h"
-#include "Styles/RenderPagesEditorStyle.h"
-#include "Widgets/Input/SButton.h"
+#include "Widgets/Input/SCheckBox.h"
 
 #define LOCTEXT_NAMESPACE "SRenderPagesPageViewer"
 
@@ -66,21 +65,30 @@ void UE::RenderPages::Private::SRenderPagesPageViewer::Construct(const FArgument
 
 TSharedRef<SWidget> UE::RenderPages::Private::SRenderPagesPageViewer::CreateViewerModeButton(const FText& ButtonText, const ERenderPagesPageViewerMode ButtonViewerMode)
 {
-	return SNew(SButton)
-		.ButtonStyle(FRenderPagesEditorStyle::Get(), "TabButton")
-		.Text(ButtonText)
-		.HAlign(HAlign_Center)
-		.VAlign(VAlign_Center)
-		.OnClicked_Lambda([this, ButtonViewerMode]() -> FReply
+	return SNew(SCheckBox)
+		.Style(FAppStyle::Get(), "PlacementBrowser.Tab")
+		.IsChecked_Lambda([this, ButtonViewerMode]() -> ECheckBoxState
 		{
+			return (ViewerMode == ButtonViewerMode) ? ECheckBoxState::Checked : ECheckBoxState::Unchecked;
+		})
+		.OnCheckStateChanged_Lambda([this, ButtonViewerMode](ECheckBoxState State)
+		{
+			if (State != ECheckBoxState::Checked)
+			{
+				return;
+			}
 			ViewerMode = ButtonViewerMode;
 			Refresh();
-			return FReply::Handled();
 		})
-		.IsEnabled_Lambda([this, ButtonViewerMode]() -> bool
-		{
-			return (ViewerMode != ButtonViewerMode);
-		});
+		[
+			SNew(SBox)
+			.HAlign(HAlign_Center)
+			.VAlign(VAlign_Center)
+			[
+				SNew(STextBlock)
+				.Text(ButtonText)
+			]
+		];
 }
 
 END_SLATE_FUNCTION_BUILD_OPTIMIZATION
