In render pages, fixed a crash that could occur when attempting to get or set the value of a remote control object when the bound object isn't available.

diff --git a/Engine/Plugins/Experimental/RenderPages/Source/RenderPages/Private/RenderPage/RenderPagePropsSource.cpp b/Engine/Plugins/Experimental/RenderPages/Source/RenderPages/Private/RenderPage/RenderPagePropsSource.cpp
--- a/Engine/Plugins/Experimental/RenderPages/Source/RenderPages/Private/RenderPage/RenderPagePropsSource.cpp
+++ b/Engine/Plugins/Experimental/RenderPages/Source/RenderPages/Private/RenderPage/RenderPagePropsSource.cpp
@@ -9,15 +9,22 @@
 
 bool /*URenderPagePropRemoteControl::*/GetObjectRef(const TSharedPtr<FRemoteControlProperty>& Field, const ERCAccess Access, FRCObjectReference& OutObjectRef)
 {
-	FRCObjectReference ObjectRef;
-	FString* ErrorText = nullptr;
-	if (!IRemoteControlModule::Get().ResolveObjectProperty(Access, Field->GetBoundObject(), Field->FieldPathInfo, ObjectRef, ErrorText))
+	if (!Field.IsValid())
 	{
-		UE_LOG(LogRenderPages, Warning, TEXT("Couldn\'t resolve object property \"%s\" in \"%s\": %s"), *Field->FieldName.ToString(), (ErrorText ? **ErrorText : TEXT("unknown")));
 		return false;
 	}
-	OutObjectRef = ObjectRef;
-	return true;
+	if (UObject* FieldBoundObject = Field->GetBoundObject(); IsValid(FieldBoundObject))
+	{
+		FRCObjectReference ObjectRef;
+		FString* ErrorText = nullptr;
+		if (IRemoteControlModule::Get().ResolveObjectProperty(Access, FieldBoundObject, Field->FieldPathInfo, ObjectRef, ErrorText))
+		{
+			OutObjectRef = ObjectRef;
+			return true;
+		}
+		UE_LOG(LogRenderPages, Warning, TEXT("Couldn\'t resolve object property \"%s\" in object \"%s\": %s"), *Field->FieldName.ToString(), *FieldBoundObject->GetPathName(), (ErrorText ? **ErrorText : TEXT("unknown")));
+	}
+	return false;
 }
 
 bool URenderPagePropRemoteControl::GetValueOfEntity(const TSharedPtr<FRemoteControlEntity>& RemoteControlEntity, TArray<uint8>& OutBinaryArray)
