- Fixed stability issues and crashes caused by the live viewport (by the way it referenced the level sequence player).
- Made it possible to wait more than 1 frame before rendering a preview (code-wise).
- Made EXR no longer be used/allowed for preview rendering.

diff --git a/Engine/Plugins/Experimental/RenderPages/Source/RenderPages/Private/RenderPage/RenderPageMoviePipelineJob.cpp b/Engine/Plugins/Experimental/RenderPages/Source/RenderPages/Private/RenderPage/RenderPageMoviePipelineJob.cpp
--- a/Engine/Plugins/Experimental/RenderPages/Source/RenderPages/Private/RenderPage/RenderPageMoviePipelineJob.cpp
+++ b/Engine/Plugins/Experimental/RenderPages/Source/RenderPages/Private/RenderPage/RenderPageMoviePipelineJob.cpp
@@ -102,7 +102,7 @@ URenderPagesMoviePipelineRenderJob* URenderPagesMoviePipelineRenderJob::Create(c
 				{
 					continue;
 				}
-				if (Cast<UMoviePipelineImageSequenceOutputBase>(Setting))
+				if (Cast<UMoviePipelineImageSequenceOutput_PNG>(Setting) || Cast<UMoviePipelineImageSequenceOutput_JPG>(Setting) || Cast<UMoviePipelineImageSequenceOutput_BMP>(Setting))
 				{
 					if (!Args.bForceOnlySingleOutput || (!bFound && (!bContainsPreferredType || Cast<UMoviePipelineImageSequenceOutput_PNG>(Setting))))
 					{
@@ -114,7 +114,7 @@ URenderPagesMoviePipelineRenderJob* URenderPagesMoviePipelineRenderJob::Create(c
 			}
 			if (Args.bForceOutputImage && !bFound)
 			{
-				if(UMoviePipelineImageSequenceOutput_PNG* NewSetting = Cast<UMoviePipelineImageSequenceOutput_PNG>(NewJob->GetConfiguration()->FindOrAddSettingByClass(UMoviePipelineImageSequenceOutput_PNG::StaticClass())))
+				if (UMoviePipelineImageSequenceOutput_PNG* NewSetting = Cast<UMoviePipelineImageSequenceOutput_PNG>(NewJob->GetConfiguration()->FindOrAddSettingByClass(UMoviePipelineImageSequenceOutput_PNG::StaticClass())))
 				{
 					NewSetting->bWriteAlpha = false;
 				}
diff --git a/Engine/Plugins/Experimental/RenderPages/Source/RenderPagesEditor/Private/UI/SRenderPagesPageViewerLive.cpp b/Engine/Plugins/Experimental/RenderPages/Source/RenderPagesEditor/Private/UI/SRenderPagesPageViewerLive.cpp
--- a/Engine/Plugins/Experimental/RenderPages/Source/RenderPagesEditor/Private/UI/SRenderPagesPageViewerLive.cpp
+++ b/Engine/Plugins/Experimental/RenderPages/Source/RenderPagesEditor/Private/UI/SRenderPagesPageViewerLive.cpp
@@ -35,19 +35,16 @@ void UE::RenderPages::Private::SRenderPagesEditorViewport::Tick(const FGeometry&
 	{
 		return;
 	}
-	ViewportClient->Invalidate();// render
+	ViewportClient->Invalidate();// causes it to rerender
 
-	if (LevelSequenceActor.IsValid())
+	if (ULevelSequencePlayer* SequencePlayer = GetSequencePlayer(); IsValid(SequencePlayer))
 	{
-		if (ULevelSequencePlayer* SequencePlayer = LevelSequenceActor->GetSequencePlayer(); IsValid(SequencePlayer))
+		SequencePlayer->SetPlaybackPosition(FMovieSceneSequencePlaybackParams(LevelSequenceTime, EUpdatePositionMethod::Play));// execute this every tick, in case any sequencer values get overwritten (by remote control props for example)
+		if (UCameraComponent* Camera = SequencePlayer->GetActiveCameraComponent(); IsValid(Camera))
 		{
-			SequencePlayer->SetPlaybackPosition(FMovieSceneSequencePlaybackParams(LevelSequenceTime, EUpdatePositionMethod::Play));// execute every tick, in case any sequencer values get overwritten (by remote control props for example)
-			if (UCameraComponent* Camera = SequencePlayer->GetActiveCameraComponent(); IsValid(Camera))
-			{
-				ViewportClient->SetViewLocation(Camera->GetComponentLocation());
-				ViewportClient->SetViewRotation(Camera->GetComponentRotation());
-				return;
-			}
+			ViewportClient->SetViewLocation(Camera->GetComponentLocation());
+			ViewportClient->SetViewRotation(Camera->GetComponentRotation());
+			return;
 		}
 	}
 
@@ -77,59 +74,105 @@ void UE::RenderPages::Private::SRenderPagesEditorViewport::Tick(const FGeometry&
 
 void UE::RenderPages::Private::SRenderPagesEditorViewport::Construct(const FArguments& InArgs)
 {
-	LevelSequenceTime = 0.0f;
 	ViewportClient = MakeShareable(new FRenderPagesEditorViewportClient(nullptr, SharedThis(this)));
+	LevelSequencePlayerWorld = nullptr;
+	LevelSequencePlayerActor = nullptr;
+	LevelSequencePlayer = nullptr;
+	LevelSequence = nullptr;
+	LevelSequenceTime = 0.0f;
 
 	SEditorViewport::Construct(SEditorViewport::FArguments());
-
-	if (UWorld* World = GetWorld(); IsValid(World))
-	{
-		FActorSpawnParameters LevelSequenceSpawnParams;
-		LevelSequenceSpawnParams.ObjectFlags = RF_Transient | RF_DuplicateTransient;
-		LevelSequenceActor = TStrongObjectPtr(World->SpawnActor<ALevelSequenceActor>(LevelSequenceSpawnParams));
-		LevelSequenceActor->PlaybackSettings.LoopCount.Value = 0;
-		LevelSequenceActor->PlaybackSettings.bAutoPlay = false;
-		LevelSequenceActor->PlaybackSettings.bPauseAtEnd = true;
-		LevelSequenceActor->PlaybackSettings.bRestoreState = true;
-	}
 }
 
 UE::RenderPages::Private::SRenderPagesEditorViewport::~SRenderPagesEditorViewport()
 {
-	if (LevelSequenceActor.IsValid())
-	{
-		if (ULevelSequencePlayer* SequencePlayer = LevelSequenceActor->GetSequencePlayer(); IsValid(SequencePlayer))
-		{
-			SequencePlayer->Stop();
-		}
-		LevelSequenceActor->Destroy();
-	}
-	LevelSequenceActor.Reset();
+	DestroySequencePlayer();
 	ViewportClient.Reset();
 }
 
 bool UE::RenderPages::Private::SRenderPagesEditorViewport::ShowSequenceFrame(ULevelSequence* InSequence, const float InTime)
 {
 	LevelSequenceTime = InTime;
-	if (!LevelSequenceActor.IsValid() || !IsValid(InSequence))
+	if (!IsValid(InSequence))
 	{
+		LevelSequence = nullptr;
+		DestroySequencePlayer();
 		return false;
 	}
-	if (LevelSequenceActor->GetSequence() != InSequence)
+	if (!IsValid(LevelSequence) || (LevelSequence != InSequence))
 	{
-		if (ULevelSequencePlayer* SequencePlayer = LevelSequenceActor->GetSequencePlayer(); IsValid(SequencePlayer))
-		{
-			SequencePlayer->Stop();
-		}
-		LevelSequenceActor->SetSequence(InSequence);
+		LevelSequence = InSequence;
+		DestroySequencePlayer();
 	}
-	if (ULevelSequencePlayer* SequencePlayer = LevelSequenceActor->GetSequencePlayer(); IsValid(SequencePlayer))
+	if (ULevelSequencePlayer* SequencePlayer = GetSequencePlayer(); IsValid(SequencePlayer))
 	{
 		SequencePlayer->SetPlaybackPosition(FMovieSceneSequencePlaybackParams(LevelSequenceTime, EUpdatePositionMethod::Play));
 	}
 	return true;
 }
 
+ULevelSequencePlayer* UE::RenderPages::Private::SRenderPagesEditorViewport::GetSequencePlayer()
+{
+	if (!IsValid(LevelSequence))
+	{
+		return nullptr;
+	}
+	if (UWorld* World = GetWorld(); IsValid(World))
+	{
+		if (IsValid(LevelSequencePlayer) && LevelSequencePlayerWorld.IsValid() && (World == LevelSequencePlayerWorld))
+		{
+			return LevelSequencePlayer;
+		}
+
+		LevelSequencePlayerWorld = nullptr;
+		LevelSequencePlayerActor = nullptr;
+		LevelSequencePlayer = nullptr;
+
+		FMovieSceneSequencePlaybackSettings PlaybackSettings;
+		PlaybackSettings.LoopCount.Value = 0;
+		PlaybackSettings.bAutoPlay = false;
+		PlaybackSettings.bPauseAtEnd = true;
+		PlaybackSettings.bRestoreState = true;
+		FLevelSequenceCameraSettings CameraSettings;
+
+		ALevelSequenceActor* PlayerActor = nullptr;
+		if (ULevelSequencePlayer* Player = ULevelSequencePlayer::CreateLevelSequencePlayer(World, LevelSequence, PlaybackSettings, PlayerActor); IsValid(Player))
+		{
+			if (IsValid(PlayerActor))
+			{
+				Player->Initialize(LevelSequence, World->PersistentLevel, PlaybackSettings, CameraSettings);
+				Player->State.AssignSequence(MovieSceneSequenceID::Root, *LevelSequence, *Player);
+				Player->SetPlaybackPosition(FMovieSceneSequencePlaybackParams(LevelSequence->GetMovieScene()->GetPlaybackRange().GetLowerBoundValue().Value, EUpdatePositionMethod::Play));
+
+				LevelSequencePlayerWorld = World;
+				LevelSequencePlayerActor = PlayerActor;
+				LevelSequencePlayer = Player;
+				return Player;
+			}
+		}
+	}
+	return nullptr;
+}
+
+void UE::RenderPages::Private::SRenderPagesEditorViewport::DestroySequencePlayer()
+{
+	TObjectPtr<ALevelSequenceActor> PlayerActor = LevelSequencePlayerActor;
+	TObjectPtr<ULevelSequencePlayer> Player = LevelSequencePlayer;
+
+	LevelSequencePlayerWorld = nullptr;
+	LevelSequencePlayerActor = nullptr;
+	LevelSequencePlayer = nullptr;
+
+	if (IsValid(Player))
+	{
+		Player->Stop();
+	}
+	if (IsValid(PlayerActor))
+	{
+		PlayerActor->Destroy();
+	}
+}
+
 
 BEGIN_SLATE_FUNCTION_BUILD_OPTIMIZATION
 
diff --git a/Engine/Plugins/Experimental/RenderPages/Source/RenderPagesEditor/Private/UI/SRenderPagesPageViewerLive.h b/Engine/Plugins/Experimental/RenderPages/Source/RenderPagesEditor/Private/UI/SRenderPagesPageViewerLive.h
--- a/Engine/Plugins/Experimental/RenderPages/Source/RenderPagesEditor/Private/UI/SRenderPagesPageViewerLive.h
+++ b/Engine/Plugins/Experimental/RenderPages/Source/RenderPagesEditor/Private/UI/SRenderPagesPageViewerLive.h
@@ -4,6 +4,7 @@
 
 #include "CoreMinimal.h"
 #include "LevelSequence.h"
+#include "LevelSequencePlayer.h"
 #include "Widgets/SCompoundWidget.h"
 #include "SEditorViewport.h"
 
@@ -62,12 +63,28 @@ namespace UE::RenderPages::Private
 		virtual bool SupportsKeyboardFocus() const override { return false; }
 		//~ End SEditorViewport Interface
 
+		ULevelSequencePlayer* GetSequencePlayer();
+		void DestroySequencePlayer();
+
 	private:
 		/** The viewport client. */
 		TSharedPtr<FRenderPagesEditorViewportClient> ViewportClient;
 
+		/** The world that the level sequence actor was spawned in. */
+		UPROPERTY()
+		TWeakObjectPtr<UWorld> LevelSequencePlayerWorld;
+
 		/** The level sequence actor we spawned to play the sequence of any given page. */
-		TStrongObjectPtr<ALevelSequenceActor> LevelSequenceActor;
+		UPROPERTY()
+		TObjectPtr<ALevelSequenceActor> LevelSequencePlayerActor;
+
+		/** The level sequence player we created to play the sequence of any given page. */
+		UPROPERTY()
+		TObjectPtr<ULevelSequencePlayer> LevelSequencePlayer;
+
+		/** The level sequence that we're currently playing. */
+		UPROPERTY()
+		TObjectPtr<ULevelSequence> LevelSequence;
 
 		/** The time of the currently playing sequence. */
 		float LevelSequenceTime;
diff --git a/Engine/Plugins/Experimental/RenderPages/Source/RenderPagesEditor/Private/UI/SRenderPagesPageViewerPreview.cpp b/Engine/Plugins/Experimental/RenderPages/Source/RenderPagesEditor/Private/UI/SRenderPagesPageViewerPreview.cpp
--- a/Engine/Plugins/Experimental/RenderPages/Source/RenderPagesEditor/Private/UI/SRenderPagesPageViewerPreview.cpp
+++ b/Engine/Plugins/Experimental/RenderPages/Source/RenderPagesEditor/Private/UI/SRenderPagesPageViewerPreview.cpp
@@ -28,9 +28,13 @@ void UE::RenderPages::Private::SRenderPagesPageViewerPreview::Tick(const FGeomet
 			UpdateRerenderButton();
 			UpdateFrameSlider();
 
-			if (bRunRenderNewPreview)
+			if (FramesUntilRenderNewPreview > 0)
 			{
-				InternalRenderNewPreview();
+				FramesUntilRenderNewPreview--;
+				if (FramesUntilRenderNewPreview <= 0)
+				{
+					InternalRenderNewPreview();
+				}
 			}
 		}
 	}
@@ -43,7 +47,7 @@ void UE::RenderPages::Private::SRenderPagesPageViewerPreview::Construct(const FA
 	BlueprintEditorWeakPtr = InBlueprintEditor;
 	SelectedPageWeakPtr = nullptr;
 	CurrentJob = nullptr;
-	bRunRenderNewPreview = false;
+	FramesUntilRenderNewPreview = 0;
 	ImageBrushEmpty = FSlateBrush();
 	ImageBrushEmpty.DrawAs = ESlateBrushDrawType::Type::NoDrawType;
 	ImageBrush = FSlateBrush();
@@ -207,12 +211,11 @@ void UE::RenderPages::Private::SRenderPagesPageViewerPreview::FrameSliderValueCh
 
 void UE::RenderPages::Private::SRenderPagesPageViewerPreview::RenderNewPreview()
 {
-	bRunRenderNewPreview = true;
+	FramesUntilRenderNewPreview = 1;
 }
 
 void UE::RenderPages::Private::SRenderPagesPageViewerPreview::InternalRenderNewPreview()
 {
-	bRunRenderNewPreview = false;
 	InternalRenderNewPreviewOfPage(SelectedPageWeakPtr.Get());
 }
 
diff --git a/Engine/Plugins/Experimental/RenderPages/Source/RenderPagesEditor/Private/UI/SRenderPagesPageViewerPreview.h b/Engine/Plugins/Experimental/RenderPages/Source/RenderPagesEditor/Private/UI/SRenderPagesPageViewerPreview.h
--- a/Engine/Plugins/Experimental/RenderPages/Source/RenderPagesEditor/Private/UI/SRenderPagesPageViewerPreview.h
+++ b/Engine/Plugins/Experimental/RenderPages/Source/RenderPagesEditor/Private/UI/SRenderPagesPageViewerPreview.h
@@ -79,8 +79,8 @@ namespace UE::RenderPages::Private
 		/** A reference to the job that's currently rendering. */
 		TWeakObjectPtr<URenderPage> SelectedPageWeakPtr;
 
-		/** True if it should call RenderNewPreview() next frame. */
-		bool bRunRenderNewPreview;
+		/** 1 if it should call RenderNewPreview() next frame, 2+ if it should subtract 1 from its value next frame, 0 and below and it won't do anything. */
+		int32 FramesUntilRenderNewPreview;
 
 	private:
 		/** True if it has rendered before since the start of this application. Used for not hiding the rendering popup during the first render (since the first render can take a lot longer due to having to compile shaders etc). */
