Fix Mac non-unity error.

diff --git a/Engine/Plugins/Experimental/RenderPages/Source/RenderPagesEditor/Private/UI/SRenderPagesPropsRemoteControl.h b/Engine/Plugins/Experimental/RenderPages/Source/RenderPagesEditor/Private/UI/SRenderPagesPropsRemoteControl.h
--- a/Engine/Plugins/Experimental/RenderPages/Source/RenderPagesEditor/Private/UI/SRenderPagesPropsRemoteControl.h
+++ b/Engine/Plugins/Experimental/RenderPages/Source/RenderPagesEditor/Private/UI/SRenderPagesPropsRemoteControl.h
@@ -6,7 +6,7 @@
 #include "RemoteControlPreset.h"
 #include "UI/SRenderPagesPropsBase.h"
 #include "RenderPage/RenderPagePropsSource.h"
-
+#include "Components/SRenderPagesRemoteControlTreeNode.h"
 
 class URenderPage;
 
