Fixed a non-unity compiler error.

diff --git a/Engine/Plugins/Experimental/RenderPages/Source/RenderPagesEditor/Private/UI/SRenderPagesPropsRemoteControl.h b/Engine/Plugins/Experimental/RenderPages/Source/RenderPagesEditor/Private/UI/SRenderPagesPropsRemoteControl.h
--- a/Engine/Plugins/Experimental/RenderPages/Source/RenderPagesEditor/Private/UI/SRenderPagesPropsRemoteControl.h
+++ b/Engine/Plugins/Experimental/RenderPages/Source/RenderPagesEditor/Private/UI/SRenderPagesPropsRemoteControl.h
@@ -4,9 +4,10 @@
 
 #include "CoreMinimal.h"
 #include "RemoteControlPreset.h"
-#include "UI/SRenderPagesPropsBase.h"
 #include "RenderPage/RenderPagePropsSource.h"
-#include "Components/SRenderPagesRemoteControlTreeNode.h"
+#include "UI/Components/SRenderPagesRemoteControlTreeNode.h"
+#include "UI/SRenderPagesPropsBase.h"
+
 
 class URenderPage;
 
@@ -15,12 +16,6 @@ namespace UE::RenderPages
 	class IRenderPageCollectionEditor;
 }
 
-namespace UE::RenderPages::Private
-{
-	struct SRenderPagesRemoteControlTreeNode;
-	struct FRenderPagesRemoteControlGenerateWidgetArgs;
-}
-
 
 namespace UE::RenderPages::Private
 {
