Made several minor fixes and improvements to the render pages plugin.

diff --git a/Engine/Plugins/Experimental/RenderPages/Source/RenderPages/Private/RenderPage/RenderPageMoviePipelineJob.cpp b/Engine/Plugins/Experimental/RenderPages/Source/RenderPages/Private/RenderPage/RenderPageMoviePipelineJob.cpp
--- a/Engine/Plugins/Experimental/RenderPages/Source/RenderPages/Private/RenderPage/RenderPageMoviePipelineJob.cpp
+++ b/Engine/Plugins/Experimental/RenderPages/Source/RenderPages/Private/RenderPage/RenderPageMoviePipelineJob.cpp
@@ -310,19 +310,21 @@ URenderPagesMoviePipelineRenderJob* URenderPagesMoviePipelineRenderJob::Create(c
 		{
 			RenderJob->Entries.Add(Page, Entry);
 
-			RenderJob->Queue->Add(UE::RenderPages::Private::FRenderPageQueueActionReturningDelay::CreateLambda([RenderJob, Page]() -> UE::RenderPages::Private::FRenderPageQueueDelay
+			RenderJob->Queue->Add(UE::RenderPages::Private::FRenderPageQueueAction::CreateLambda([RenderJob, Page]()
 			{
 				RenderJob->PreviousFrameLimitSettings = UE::RenderPages::Private::FRenderPagesUtils::DisableFpsLimit();
 				RenderJob->PreviousPageProps = UE::RenderPages::IRenderPagesModule::Get().GetManager().ApplyPagePropValues(RenderJob->PageCollection, Page.Get());
-				return UE::RenderPages::Private::FRenderPageQueueDelay::Frames(1 + Page->GetWaitFramesBeforeRendering());
 			}));
+			
+			RenderJob->Queue->DelayFrames(1 + Page->GetWaitFramesBeforeRendering());
 
-			RenderJob->Queue->Add(UE::RenderPages::Private::FRenderPageQueueActionReturningDelay::CreateLambda([RenderJob]() -> UE::RenderPages::Private::FRenderPageQueueDelay
+			RenderJob->Queue->Add(UE::RenderPages::Private::FRenderPageQueueAction::CreateLambda([RenderJob]()
 			{
 				UE::RenderPages::Private::FRenderPagesUtils::RestoreFpsLimit(RenderJob->PreviousFrameLimitSettings);
 				RenderJob->PreviousFrameLimitSettings = FRenderPagePreviousEngineFpsSettings();
-				return UE::RenderPages::Private::FRenderPageQueueDelay::Frames(1);
 			}));
+			
+			RenderJob->Queue->DelayFrames(1);
 
 			RenderJob->Queue->Add(UE::RenderPages::Private::FRenderPageQueueActionReturningDelayFuture::CreateLambda([Entry]()-> TSharedFuture<void>
 			{
diff --git a/Engine/Plugins/Experimental/RenderPages/Source/RenderPages/Private/Utils/RenderPageQueue.cpp b/Engine/Plugins/Experimental/RenderPages/Source/RenderPages/Private/Utils/RenderPageQueue.cpp
--- a/Engine/Plugins/Experimental/RenderPages/Source/RenderPages/Private/Utils/RenderPageQueue.cpp
+++ b/Engine/Plugins/Experimental/RenderPages/Source/RenderPages/Private/Utils/RenderPageQueue.cpp
@@ -28,7 +28,7 @@ void UE::RenderPages::Private::FRenderPageQueue::Tick(float DeltaTime)
 		{
 			return;
 		}
-		DelayFramesOrSeconds(DelayRemainingFutureReturningDelay.Get());
+		QueueDelay(DelayRemainingFutureReturningDelay.Get());
 		DelayRemainingFutureReturningDelay = TFuture<FRenderPageQueueDelay>();
 	}
 
@@ -110,7 +110,7 @@ bool UE::RenderPages::Private::FRenderPageQueue::ExecuteNextEntry()
 
 		if (Entry.ActionReturningDelay.IsBound())
 		{
-			DelayFramesOrSeconds(Entry.ActionReturningDelay.Execute());
+			QueueDelay(Entry.ActionReturningDelay.Execute());
 		}
 
 		if (Entry.ActionReturningDelayFuture.IsBound())
diff --git a/Engine/Plugins/Experimental/RenderPages/Source/RenderPages/Private/Utils/RenderPageQueue.h b/Engine/Plugins/Experimental/RenderPages/Source/RenderPages/Private/Utils/RenderPageQueue.h
--- a/Engine/Plugins/Experimental/RenderPages/Source/RenderPages/Private/Utils/RenderPageQueue.h
+++ b/Engine/Plugins/Experimental/RenderPages/Source/RenderPages/Private/Utils/RenderPageQueue.h
@@ -5,7 +5,7 @@
 #include "CoreMinimal.h"
 #include "Async/Future.h"
 #include "Containers/List.h"
-#include "TickableEditorObject.h"
+#include "Tickable.h"
 
 
 namespace UE::RenderPages::Private
@@ -20,7 +20,7 @@ namespace UE::RenderPages::Private
 		static FRenderPageQueueDelay Seconds(const double Seconds) { return FRenderPageQueueDelay(Seconds); }
 		static FRenderPageQueueDelay FramesOrSeconds(const int64 Frames, const double Seconds) { return FRenderPageQueueDelay(Frames, Seconds); }
 
-	public:
+	private:
 		FRenderPageQueueDelay()
 		{}
 
@@ -108,17 +108,17 @@ namespace UE::RenderPages::Private
 		/** Queues the given action. */
 		void Add(const FRenderPageQueueEntry& Entry) { QueuedEntries.AddTail(Entry); }
 
-		/** Delays the given number of frames. */
-		void DelayFrames(const int64 Frames) { QueuedDelays.AddTail(FRenderPageQueueDelay(Frames)); }
+		/** Queues the given delay. */
+		void Delay(const FRenderPageQueueDelay& Delay) { Add(FRenderPageQueueActionReturningDelay::CreateLambda([Delay]() -> FRenderPageQueueDelay { return Delay; })); }
 
-		/** Delays the given number of seconds. */
-		void DelaySeconds(const double Seconds) { QueuedDelays.AddTail(FRenderPageQueueDelay(Seconds)); }
+		/** Queues the given delay, which will wait for the given number of frames. */
+		void DelayFrames(const int64 Frames) { Delay(FRenderPageQueueDelay::Frames(Frames)); }
 
-		/** Delays the given number of frames or seconds, whatever takes the longest. */
-		void DelayFramesOrSeconds(const int64 Frames, const double Seconds) { QueuedDelays.AddTail(FRenderPageQueueDelay(Frames, Seconds)); }
+		/** Queues the given delay, which will wait for the given number of seconds. */
+		void DelaySeconds(const double Seconds) { Delay(FRenderPageQueueDelay::Seconds(Seconds)); }
 
-		/** Delays the given number of frames or seconds, whatever takes the longest. */
-		void DelayFramesOrSeconds(const FRenderPageQueueDelay& Delay) { QueuedDelays.AddTail(Delay); }
+		/** Queues the given delay, which will wait for the given number of frames or seconds, whatever takes the longest. */
+		void DelayFramesOrSeconds(const int64 Frames, const double Seconds) { Delay(FRenderPageQueueDelay::FramesOrSeconds(Frames, Seconds)); }
 
 		/** Starts the execution of this queue. */
 		void Start();
@@ -139,6 +139,9 @@ namespace UE::RenderPages::Private
 		/** Executes the next entry (action), returns true if it found and executed an entry, returns false if there were no queued up entries. */
 		bool ExecuteNextEntry();
 
+		/** Adds the delay to the queued delays. */
+		void QueueDelay(const FRenderPageQueueDelay& Delay) { QueuedDelays.AddTail(Delay); }
+
 	protected:
 		/** The queued up entries (actions). */
 		TDoubleLinkedList<FRenderPageQueueEntry> QueuedEntries;
diff --git a/Engine/Plugins/Experimental/RenderPages/Source/RenderPagesEditor/Private/UI/Components/SRenderPagesRemoteControlEntity.cpp b/Engine/Plugins/Experimental/RenderPages/Source/RenderPagesEditor/Private/UI/Components/SRenderPagesRemoteControlEntity.cpp
--- a/Engine/Plugins/Experimental/RenderPages/Source/RenderPagesEditor/Private/UI/Components/SRenderPagesRemoteControlEntity.cpp
+++ b/Engine/Plugins/Experimental/RenderPages/Source/RenderPagesEditor/Private/UI/Components/SRenderPagesRemoteControlEntity.cpp
@@ -50,7 +50,7 @@ TSharedRef<SWidget> UE::RenderPages::Private::SRenderPagesRemoteControlEntity::C
 		[
 			SNew(SInlineEditableTextBlock)
 			.Text(FText::FromName(CachedLabel))
-			.IsReadOnly_Lambda([this]() { return true; /*!bEditMode.Get();*/ })
+			.IsReadOnly(true)
 		];
 
 	Args.ValueWidget = ValueWidget;
diff --git a/Engine/Plugins/Experimental/RenderPages/Source/RenderPagesEditor/Private/UI/Components/SRenderPagesRemoteControlField.cpp b/Engine/Plugins/Experimental/RenderPages/Source/RenderPagesEditor/Private/UI/Components/SRenderPagesRemoteControlField.cpp
--- a/Engine/Plugins/Experimental/RenderPages/Source/RenderPagesEditor/Private/UI/Components/SRenderPagesRemoteControlField.cpp
+++ b/Engine/Plugins/Experimental/RenderPages/Source/RenderPagesEditor/Private/UI/Components/SRenderPagesRemoteControlField.cpp
@@ -137,7 +137,8 @@ void UE::RenderPages::Private::SRenderPagesRemoteControlField::RefreshValue()
 				ChildWidgets.Add(SNew(SRenderPagesRemoteControlFieldChildNode, ChildNode, ColumnSizeData));
 			}
 
-			//TODO:  still causes the value widgets (like the color wheel) to disconnect when this function is called,  maybe Node->CreateNodeWidgets() should be cached
+			//TODO:  still causes the value widgets (like the color wheel) to disconnect when this function is called,  have to somehow prevent this disconnection from happening,
+			//       skipping this line all-together causes the value widgets to not respond anymore to user input,  same happens when trying to cache and reuse the output of Node->CreateNodeWidgets()
 			ChildSlot.AttachWidget(MakeFieldWidget(ExposedFieldUtils::CreateNodeValueWidget(Node->CreateNodeWidgets())));
 		}
 		else
@@ -218,4 +219,4 @@ void UE::RenderPages::Private::SRenderPagesRemoteControlFieldChildNode::Construc
 }
 
 
-#undef LOCTEXT_NAMESPACE /*RemoteControlPanel*/
+#undef LOCTEXT_NAMESPACE
diff --git a/Engine/Plugins/Experimental/RenderPages/Source/RenderPagesEditor/Private/UI/Components/SRenderPagesRemoteControlTreeNode.cpp b/Engine/Plugins/Experimental/RenderPages/Source/RenderPagesEditor/Private/UI/Components/SRenderPagesRemoteControlTreeNode.cpp
--- a/Engine/Plugins/Experimental/RenderPages/Source/RenderPagesEditor/Private/UI/Components/SRenderPagesRemoteControlTreeNode.cpp
+++ b/Engine/Plugins/Experimental/RenderPages/Source/RenderPagesEditor/Private/UI/Components/SRenderPagesRemoteControlTreeNode.cpp
@@ -90,4 +90,4 @@ void UE::RenderPages::Private::SRenderPagesRemoteControlTreeNode::SetColumnWidth
 }
 
 
-#undef LOCTEXT_NAMESPACE /*RemoteControlPanelNode*/
+#undef LOCTEXT_NAMESPACE
