Implemented a work-around to get camera location/rotation/fov to work in the live viewer.

diff --git a/Engine/Plugins/Experimental/RenderPages/Source/RenderPages/Private/Utils/RenderPageLevelSequencePlayer.cpp b/Engine/Plugins/Experimental/RenderPages/Source/RenderPages/Private/Utils/RenderPageLevelSequencePlayer.cpp
new file mode 100644
--- /dev/null
+++ b/Engine/Plugins/Experimental/RenderPages/Source/RenderPages/Private/Utils/RenderPageLevelSequencePlayer.cpp
@@ -0,0 +1,71 @@
+﻿// Copyright Epic Games, Inc. All Rights Reserved.
+
+#include "Utils/RenderPageLevelSequencePlayer.h"
+#include "MovieSceneCommonHelpers.h"
+#include "Camera/CameraComponent.h"
+
+
+ULevelSequencePlayer* URenderPageLevelSequencePlayer::CreateLevelSequencePlayer(UObject* WorldContextObject, ULevelSequence* InLevelSequence, FMovieSceneSequencePlaybackSettings Settings, ALevelSequenceActor*& OutActor)
+{
+	// entire function is copied from ULevelSequencePlayer::CreateLevelSequencePlayer (LevelSequencePlayer.cpp, lines 42-75)
+
+	if (InLevelSequence == nullptr)
+	{
+		return nullptr;
+	}
+
+	UWorld* World = GEngine->GetWorldFromContextObject(WorldContextObject, EGetWorldErrorMode::LogAndReturnNull);
+	if (World == nullptr || World->bIsTearingDown)
+	{
+		return nullptr;
+	}
+
+	FActorSpawnParameters SpawnParams;
+	SpawnParams.SpawnCollisionHandlingOverride = ESpawnActorCollisionHandlingMethod::AlwaysSpawn;
+	SpawnParams.ObjectFlags |= RF_Transient;
+	SpawnParams.bAllowDuringConstructionScript = true;
+
+	// Defer construction for autoplay so that BeginPlay() is called
+	SpawnParams.bDeferConstruction = true;
+
+	ALevelSequenceActor* Actor = World->SpawnActor<ARenderPageLevelLevelSequenceActor>(SpawnParams); // line contains changes
+
+	Actor->PlaybackSettings = Settings;
+	Actor->SetSequence(InLevelSequence);
+
+	Actor->InitializePlayer();
+	OutActor = Actor;
+
+	FTransform DefaultTransform;
+	Actor->FinishSpawning(DefaultTransform);
+
+	return Actor->SequencePlayer;
+}
+
+void URenderPageLevelSequencePlayer::UpdateCameraCut(UObject* CameraObject, const EMovieSceneCameraCutParams& CameraCutParams)
+{
+	CachedCameraComponent = MovieSceneHelpers::CameraComponentFromRuntimeObject(CameraObject);
+
+	ULevelSequencePlayer::UpdateCameraCut(CameraObject, CameraCutParams);
+}
+
+
+ARenderPageLevelLevelSequenceActor::ARenderPageLevelLevelSequenceActor(const FObjectInitializer& Init)
+	: Super(Init)
+{
+	// copied from ALevelSequenceActor::EndPlay (LevelSequenceActor.cpp, lines 184-196)
+	if (SequencePlayer)
+	{
+		SequencePlayer->Stop();
+		SequencePlayer->OnPlay.RemoveAll(this);
+		SequencePlayer->OnPlayReverse.RemoveAll(this);
+		SequencePlayer->OnStop.RemoveAll(this);
+		SequencePlayer->TearDown();
+	}
+
+	// copied from ALevelSequenceActor constructor (LevelSequenceActor.cpp, lines 71-75)
+	SequencePlayer = Init.CreateDefaultSubobject<URenderPageLevelSequencePlayer>(this, "RenderPageAnimationPlayer"); // line contains changes
+	SequencePlayer->OnPlay.AddDynamic(this, &ALevelSequenceActor::ShowBurnin);
+	SequencePlayer->OnPlayReverse.AddDynamic(this, &ALevelSequenceActor::ShowBurnin);
+	SequencePlayer->OnStop.AddDynamic(this, &ALevelSequenceActor::HideBurnin);
+}
diff --git a/Engine/Plugins/Experimental/RenderPages/Source/RenderPages/Public/Utils/RenderPageLevelSequencePlayer.h b/Engine/Plugins/Experimental/RenderPages/Source/RenderPages/Public/Utils/RenderPageLevelSequencePlayer.h
new file mode 100644
--- /dev/null
+++ b/Engine/Plugins/Experimental/RenderPages/Source/RenderPages/Public/Utils/RenderPageLevelSequencePlayer.h
@@ -0,0 +1,38 @@
+﻿// Copyright Epic Games, Inc. All Rights Reserved.
+
+#pragma once
+
+#include "CoreMinimal.h"
+#include "LevelSequenceActor.h"
+#include "LevelSequencePlayer.h"
+#include "RenderPageLevelSequencePlayer.generated.h"
+
+
+/**
+ * This class overrides a function in ULevelSequencePlayer to fix an issue with obtaining the camera while running the level sequence player in the editor.
+ */
+UCLASS()
+class RENDERPAGES_API URenderPageLevelSequencePlayer : public ULevelSequencePlayer
+{
+public:
+	GENERATED_BODY()
+
+public:
+	static ULevelSequencePlayer* CreateLevelSequencePlayer(UObject* WorldContextObject, ULevelSequence* LevelSequence, FMovieSceneSequencePlaybackSettings Settings, ALevelSequenceActor*& OutActor);
+
+protected:
+	virtual void UpdateCameraCut(UObject* CameraObject, const EMovieSceneCameraCutParams& CameraCutParams) override;
+};
+
+
+/**
+ * This class overrides the constructor in ALevelSequenceActor so it will use a URenderPageLevelSequencePlayer instance instead of a ULevelSequencePlayer instance.
+ */
+UCLASS()
+class RENDERPAGES_API ARenderPageLevelLevelSequenceActor : public ALevelSequenceActor
+{
+public:
+	GENERATED_BODY()
+
+	ARenderPageLevelLevelSequenceActor(const FObjectInitializer& Init);
+};
diff --git a/Engine/Plugins/Experimental/RenderPages/Source/RenderPages/RenderPages.Build.cs b/Engine/Plugins/Experimental/RenderPages/Source/RenderPages/RenderPages.Build.cs
--- a/Engine/Plugins/Experimental/RenderPages/Source/RenderPages/RenderPages.Build.cs
+++ b/Engine/Plugins/Experimental/RenderPages/Source/RenderPages/RenderPages.Build.cs
@@ -14,6 +14,7 @@ public class RenderPages : ModuleRules
 				"Core",
 				"CoreUObject",
 				"Engine",
+				"LevelSequence",
 				"RemoteControl",
 			}
 		);
diff --git a/Engine/Plugins/Experimental/RenderPages/Source/RenderPagesEditor/Private/UI/SRenderPagesPageViewerLive.cpp b/Engine/Plugins/Experimental/RenderPages/Source/RenderPagesEditor/Private/UI/SRenderPagesPageViewerLive.cpp
--- a/Engine/Plugins/Experimental/RenderPages/Source/RenderPagesEditor/Private/UI/SRenderPagesPageViewerLive.cpp
+++ b/Engine/Plugins/Experimental/RenderPages/Source/RenderPagesEditor/Private/UI/SRenderPagesPageViewerLive.cpp
@@ -11,6 +11,7 @@
 #include "RenderPage/RenderPageCollection.h"
 #include "SlateOptMacros.h"
 #include "UI/Components/SRenderPagesPageViewerFrameSlider.h"
+#include "Utils/RenderPageLevelSequencePlayer.h"
 #include "Widgets/Layout/SScaleBox.h"
 
 #define LOCTEXT_NAMESPACE "SRenderPagesPageViewerLive"
@@ -44,9 +45,11 @@ void UE::RenderPages::Private::SRenderPagesEditorViewport::Tick(const FGeometry&
 		{
 			ViewportClient->SetViewLocation(Camera->GetComponentLocation());
 			ViewportClient->SetViewRotation(Camera->GetComponentRotation());
+			ViewportClient->ViewFOV = Camera->FieldOfView;
 			return;
 		}
 	}
+	ViewportClient->ViewFOV = 90;
 
 	if (UWorld* World = GetWorld(); IsValid(World))
 	{
@@ -136,7 +139,7 @@ ULevelSequencePlayer* UE::RenderPages::Private::SRenderPagesEditorViewport::GetS
 		FLevelSequenceCameraSettings CameraSettings;
 
 		ALevelSequenceActor* PlayerActor = nullptr;
-		if (ULevelSequencePlayer* Player = ULevelSequencePlayer::CreateLevelSequencePlayer(World, LevelSequence, PlaybackSettings, PlayerActor); IsValid(Player))
+		if (ULevelSequencePlayer* Player = URenderPageLevelSequencePlayer::CreateLevelSequencePlayer(World, LevelSequence, PlaybackSettings, PlayerActor); IsValid(Player))
 		{
 			if (IsValid(PlayerActor))
 			{
