- Made the render pages UI largely hidden during batch rendering (to speed up the rendering).
- Made it so that canceling one render will cancel all the remaining queued-up renders as well.
- Removed the custom blueprint graph classes, as we want to use regular blueprint graphs (the custom one didn't support any functions/macros/etc).
- Made it possible to have renderpages events in the blueprint graph, also made the event nodes automatically added to the graph for the user.
- Remove the camera work-around needed for the live viewer (the fix has been accepted into the engine).
- Added the pre and post render blueprint events.
- Made most of the render page's functions available in blueprint.
- Cleaned up the code, removed the auto compile button (which didn't do anything).
- Fixed several issues with the renderpages blueprint graph.

diff --git a/Engine/Plugins/Experimental/RenderPages/Source/RenderPages/Private/RenderPage/RenderPageCollection.cpp b/Engine/Plugins/Experimental/RenderPages/Source/RenderPages/Private/RenderPage/RenderPageCollection.cpp
--- a/Engine/Plugins/Experimental/RenderPages/Source/RenderPages/Private/RenderPage/RenderPageCollection.cpp
+++ b/Engine/Plugins/Experimental/RenderPages/Source/RenderPages/Private/RenderPage/RenderPageCollection.cpp
@@ -504,6 +504,16 @@ void URenderPageCollection::PostLoad()
 	LoadValuesFromCDO();
 }
 
+void URenderPageCollection::PreRender(URenderPage* Page)
+{
+	ReceivePreRender(Page);
+}
+
+void URenderPageCollection::PostRender(URenderPage* Page)
+{
+	ReceivePostRender(Page);
+}
+
 void URenderPageCollection::CopyValuesToOrFromCDO(const bool bToCDO)
 {
 	if (HasAnyFlags(RF_ClassDefaultObject))
diff --git a/Engine/Plugins/Experimental/RenderPages/Source/RenderPages/Private/RenderPage/RenderPageMoviePipelineJob.cpp b/Engine/Plugins/Experimental/RenderPages/Source/RenderPages/Private/RenderPage/RenderPageMoviePipelineJob.cpp
--- a/Engine/Plugins/Experimental/RenderPages/Source/RenderPages/Private/RenderPage/RenderPageMoviePipelineJob.cpp
+++ b/Engine/Plugins/Experimental/RenderPages/Source/RenderPages/Private/RenderPage/RenderPageMoviePipelineJob.cpp
@@ -41,6 +41,7 @@ URenderPagesMoviePipelineRenderJobEntry* URenderPagesMoviePipelineRenderJobEntry
 	RenderJobEntry->ExecutorJob = nullptr;
 	RenderJobEntry->Status = TEXT("Skipped");
 	RenderJobEntry->bCanExecute = false;
+	RenderJobEntry->bCanceled = false;
 	RenderJobEntry->Promise = MakeShared<TPromise<void>>();
 	RenderJobEntry->Promise->SetValue();
 	RenderJobEntry->PromiseFuture = RenderJobEntry->Promise->GetFuture().Share();
@@ -219,6 +220,13 @@ TSharedFuture<void> URenderPagesMoviePipelineRenderJobEntry::Execute()
 		Promise.Reset();
 		return PromiseFuture;
 	}
+	if (bCanceled)
+	{
+		Status = TEXT("Canceled");
+		Promise->SetValue();
+		Promise.Reset();
+		return PromiseFuture;
+	}
 
 	AddToRoot();
 	if (ILevelSequenceEditorModule* LevelSequenceEditorModule = FModuleManager::GetModulePtr<ILevelSequenceEditorModule>("LevelSequenceEditor"); LevelSequenceEditorModule)
@@ -233,7 +241,7 @@ TSharedFuture<void> URenderPagesMoviePipelineRenderJobEntry::Execute()
 
 void URenderPagesMoviePipelineRenderJobEntry::Cancel()
 {
-	bCanExecute = false;
+	bCanceled = true;
 	if (Executor->IsRendering())
 	{
 		Executor->CancelAllJobs();
@@ -275,7 +283,8 @@ void URenderPagesMoviePipelineRenderJobEntry::ExecuteFinished(UMoviePipelineExec
 	{
 		LevelSequenceEditorModule->OnComputePlaybackContext().RemoveAll(this);
 	}
-	Status = (bSuccess ? TEXT("Done") : TEXT("Failed"));
+	bCanceled = (bCanceled || !bSuccess);
+	Status = (bCanceled ? TEXT("Canceled") : TEXT("Done"));
 	if (Promise.IsValid())
 	{
 		Promise->SetValue();
@@ -302,43 +311,86 @@ URenderPagesMoviePipelineRenderJob* URenderPagesMoviePipelineRenderJob::Create(c
 	RenderJob->Queue = MakeShareable(new UE::RenderPages::Private::FRenderPageQueue);
 	RenderJob->PageCollection = Args.PageCollection;
 	RenderJob->bCanceled = false;
+	RenderJob->bRanPreRender = false;
 
 
+	RenderJob->Queue->Add(UE::RenderPages::Private::FRenderPageQueueAction::CreateLambda([RenderJob]()
+	{
+		RenderJob->PreviousFrameLimitSettings = UE::RenderPages::Private::FRenderPagesUtils::DisableFpsLimit();
+	}));
+	RenderJob->Queue->DelayFrames(1);
+
 	for (const TObjectPtr<URenderPage> Page : Args.Pages)
 	{
 		if (URenderPagesMoviePipelineRenderJobEntry* Entry = URenderPagesMoviePipelineRenderJobEntry::Create(RenderJob, Page, Args); IsValid(Entry))
 		{
 			RenderJob->Entries.Add(Page, Entry);
 
-			RenderJob->Queue->Add(UE::RenderPages::Private::FRenderPageQueueAction::CreateLambda([RenderJob, Page]()
+			RenderJob->Queue->Add(UE::RenderPages::Private::FRenderPageQueueActionReturningDelay::CreateLambda([RenderJob, Page]() -> UE::RenderPages::Private::FRenderPageQueueDelay
 			{
-				RenderJob->PreviousFrameLimitSettings = UE::RenderPages::Private::FRenderPagesUtils::DisableFpsLimit();
-				RenderJob->PreviousPageProps = UE::RenderPages::IRenderPagesModule::Get().GetManager().ApplyPagePropValues(RenderJob->PageCollection, Page.Get());
+				if (!RenderJob->IsCanceled())
+				{
+					RenderJob->bRanPreRender = true;
+					RenderJob->PageCollection->PreRender(Page);
+					return UE::RenderPages::Private::FRenderPageQueueDelay::Frames(1);
+				}
+				return nullptr;
 			}));
-			
-			RenderJob->Queue->DelayFrames(1 + Page->GetWaitFramesBeforeRendering());
 
-			RenderJob->Queue->Add(UE::RenderPages::Private::FRenderPageQueueAction::CreateLambda([RenderJob]()
+			RenderJob->Queue->Add(UE::RenderPages::Private::FRenderPageQueueActionReturningDelay::CreateLambda([RenderJob, Page]() -> UE::RenderPages::Private::FRenderPageQueueDelay
 			{
-				UE::RenderPages::Private::FRenderPagesUtils::RestoreFpsLimit(RenderJob->PreviousFrameLimitSettings);
-				RenderJob->PreviousFrameLimitSettings = FRenderPagePreviousEngineFpsSettings();
+				if (!RenderJob->IsCanceled())
+				{
+					RenderJob->PreviousPageProps = UE::RenderPages::IRenderPagesModule::Get().GetManager().ApplyPagePropValues(RenderJob->PageCollection, Page.Get());
+					return UE::RenderPages::Private::FRenderPageQueueDelay::Frames(1 + Page->GetWaitFramesBeforeRendering());
+				}
+				return nullptr;
 			}));
-			
-			RenderJob->Queue->DelayFrames(1);
 
 			RenderJob->Queue->Add(UE::RenderPages::Private::FRenderPageQueueActionReturningDelayFuture::CreateLambda([Entry]()-> TSharedFuture<void>
 			{
 				return Entry->Execute();
 			}));
 
-			RenderJob->Queue->Add(UE::RenderPages::Private::FRenderPageQueueAction::CreateLambda([RenderJob]()
+			RenderJob->Queue->Add(UE::RenderPages::Private::FRenderPageQueueActionReturningDelay::CreateLambda([RenderJob]() -> UE::RenderPages::Private::FRenderPageQueueDelay
+			{
+				if (!RenderJob->PreviousPageProps.IsEmpty())
+				{
+					UE::RenderPages::IRenderPagesModule::Get().GetManager().RestorePagePropValues(RenderJob->PreviousPageProps);
+					RenderJob->PreviousPageProps = FRenderPageManagerPreviousPagePropValues();
+					return UE::RenderPages::Private::FRenderPageQueueDelay::Frames(1);
+				}
+				return nullptr;
+			}));
+
+			RenderJob->Queue->Add(UE::RenderPages::Private::FRenderPageQueueActionReturningDelay::CreateLambda([RenderJob, Page]() -> UE::RenderPages::Private::FRenderPageQueueDelay
 			{
-				UE::RenderPages::IRenderPagesModule::Get().GetManager().RestorePagePropValues(RenderJob->PreviousPageProps);
-				RenderJob->PreviousPageProps = FRenderPageManagerPreviousPagePropValues();
+				if (RenderJob->bRanPreRender)
+				{
+					RenderJob->bRanPreRender = false;
+					RenderJob->PageCollection->PostRender(Page);
+					return UE::RenderPages::Private::FRenderPageQueueDelay::Frames(1);
+				}
+				return nullptr;
+			}));
+
+			RenderJob->Queue->Add(UE::RenderPages::Private::FRenderPageQueueAction::CreateLambda([RenderJob, Entry]()
+			{
+				if (!RenderJob->IsCanceled() && Entry->IsCanceled())
+				{
+					RenderJob->Cancel();
+				}
 			}));
 		}
 	}
 
+	RenderJob->Queue->Add(UE::RenderPages::Private::FRenderPageQueueAction::CreateLambda([RenderJob]()
+	{
+		UE::RenderPages::Private::FRenderPagesUtils::RestoreFpsLimit(RenderJob->PreviousFrameLimitSettings);
+		RenderJob->PreviousFrameLimitSettings = FRenderPagePreviousEngineFpsSettings();
+	}));
+
+
 	if (RenderJob->Entries.Num() <= 0)
 	{
 		return nullptr;
diff --git a/Engine/Plugins/Experimental/RenderPages/Source/RenderPages/Private/Utils/RenderPageLevelSequencePlayer.cpp b/Engine/Plugins/Experimental/RenderPages/Source/RenderPages/Private/Utils/RenderPageLevelSequencePlayer.cpp
deleted file mode 100644
diff --git a/Engine/Plugins/Experimental/RenderPages/Source/RenderPages/Private/Utils/RenderPageQueue.h b/Engine/Plugins/Experimental/RenderPages/Source/RenderPages/Private/Utils/RenderPageQueue.h
--- a/Engine/Plugins/Experimental/RenderPages/Source/RenderPages/Private/Utils/RenderPageQueue.h
+++ b/Engine/Plugins/Experimental/RenderPages/Source/RenderPages/Private/Utils/RenderPageQueue.h
@@ -20,6 +20,10 @@ namespace UE::RenderPages::Private
 		static FRenderPageQueueDelay Seconds(const double Seconds) { return FRenderPageQueueDelay(Seconds); }
 		static FRenderPageQueueDelay FramesOrSeconds(const int64 Frames, const double Seconds) { return FRenderPageQueueDelay(Frames, Seconds); }
 
+	public:
+		FRenderPageQueueDelay(TYPE_OF_NULLPTR)
+		{}
+
 	private:
 		FRenderPageQueueDelay()
 		{}
diff --git a/Engine/Plugins/Experimental/RenderPages/Source/RenderPages/Public/RenderPage/RenderPageCollection.h b/Engine/Plugins/Experimental/RenderPages/Source/RenderPages/Public/RenderPage/RenderPageCollection.h
--- a/Engine/Plugins/Experimental/RenderPages/Source/RenderPages/Public/RenderPage/RenderPageCollection.h
+++ b/Engine/Plugins/Experimental/RenderPages/Source/RenderPages/Public/RenderPage/RenderPageCollection.h
@@ -79,62 +79,127 @@ public:
 	TOptional<double> GetDurationInSeconds() const;
 
 	/** Gets the aspect ratio that this page will be rendered in. */
+	UFUNCTION(BlueprintPure, Category="Render Pages | Page")
 	double GetOutputAspectRatio() const;
 
 	/** Checks whether the page contains data that matches the search terms. */
+	UFUNCTION(BlueprintPure, Category="Render Pages | Page")
 	bool MatchesSearchTerm(const FString& SearchTerm) const;
 
 public:
+	UFUNCTION(BlueprintPure, Category="Render Pages | Page")
 	FGuid GetId() const { return Id; }
+
+	UFUNCTION(BlueprintCallable, Category="Render Pages | Page")
 	void GenerateNewId() { Id = FGuid::NewGuid(); }
 
-	uint32 GetWaitFramesBeforeRendering() const { return WaitFramesBeforeRendering; }
-	void SetWaitFramesBeforeRendering(const uint32 NewWaitFramesBeforeRendering) { WaitFramesBeforeRendering = NewWaitFramesBeforeRendering; }
 
+	UFUNCTION(BlueprintPure, Category="Render Pages | Page")
+	int32 GetWaitFramesBeforeRendering() const { return WaitFramesBeforeRendering; }
+
+	UFUNCTION(BlueprintCallable, Category="Render Pages | Page")
+	void SetWaitFramesBeforeRendering(const int32 NewWaitFramesBeforeRendering) { WaitFramesBeforeRendering = FMath::Max<int32>(0, NewWaitFramesBeforeRendering); }
+
+
+	UFUNCTION(BlueprintPure, Category="Render Pages | Page")
 	ULevelSequence* GetSequence() const { return Sequence; }
+
+	UFUNCTION(BlueprintCallable, Category="Render Pages | Page")
 	void SetSequence(ULevelSequence* NewSequence) { Sequence = NewSequence; }
 
+
+	UFUNCTION(BlueprintPure, Category="Render Pages | Page")
 	bool GetIsCustomStartFrame() const { return bOverrideStartFrame; }
-	void SetIsCustomStartFrame(bool bNewOverrideStartFrame) { bOverrideStartFrame = bNewOverrideStartFrame; }
+
+	UFUNCTION(BlueprintCallable, Category="Render Pages | Page")
+	void SetIsCustomStartFrame(const bool bNewOverrideStartFrame) { bOverrideStartFrame = bNewOverrideStartFrame; }
+
+	UFUNCTION(BlueprintPure, Category="Render Pages | Page")
 	int32 GetCustomStartFrame() const { return CustomStartFrame; }
+
+	UFUNCTION(BlueprintCallable, Category="Render Pages | Page")
 	void SetCustomStartFrame(const int32 NewCustomStartFrame) { CustomStartFrame = NewCustomStartFrame; }
 
+
+	UFUNCTION(BlueprintPure, Category="Render Pages | Page")
 	bool GetIsCustomEndFrame() const { return bOverrideEndFrame; }
-	void SetIsCustomEndFrame(bool bNewOverrideEndFrame) { bOverrideEndFrame = bNewOverrideEndFrame; }
+
+	UFUNCTION(BlueprintCallable, Category="Render Pages | Page")
+	void SetIsCustomEndFrame(const bool bNewOverrideEndFrame) { bOverrideEndFrame = bNewOverrideEndFrame; }
+
+	UFUNCTION(BlueprintPure, Category="Render Pages | Page")
 	int32 GetCustomEndFrame() const { return CustomEndFrame; }
+
+	UFUNCTION(BlueprintCallable, Category="Render Pages | Page")
 	void SetCustomEndFrame(const int32 NewCustomEndFrame) { CustomEndFrame = NewCustomEndFrame; }
 
+
+	UFUNCTION(BlueprintPure, Category="Render Pages | Page")
 	bool GetIsCustomResolution() const { return bOverrideResolution; }
-	void SetIsCustomResolution(bool bNewOverrideResolution) { bOverrideResolution = bNewOverrideResolution; }
+
+	UFUNCTION(BlueprintCallable, Category="Render Pages | Page")
+	void SetIsCustomResolution(const bool bNewOverrideResolution) { bOverrideResolution = bNewOverrideResolution; }
+
+	UFUNCTION(BlueprintPure, Category="Render Pages | Page")
 	FIntPoint GetCustomResolution() const { return CustomResolution; }
+
+	UFUNCTION(BlueprintCallable, Category="Render Pages | Page")
 	void SetCustomResolution(const FIntPoint NewCustomResolution) { CustomResolution = NewCustomResolution; }
 
+
+	UFUNCTION(BlueprintPure, Category="Render Pages | Page")
 	FString GetPageId() const { return PageId; }
+
 	static FString PurgePageIdOrReturnEmptyString(const FString& NewPageId);
 	static FString PurgePageId(const FString& NewPageId);
 	static FString PurgePageIdOrGenerateUniqueId(URenderPageCollection* PageCollection, const FString& NewPageId);
+
+	UFUNCTION(BlueprintCallable, Category="Render Pages | Page")
 	void SetPageId(const FString& NewPageId) { PageId = PurgePageId(NewPageId); }
 	void SetPageIdRaw(const FString& NewPageId) { PageId = NewPageId; }
 
+
+	UFUNCTION(BlueprintPure, Category="Render Pages | Page")
 	FString GetPageName() const { return PageName; }
+
 	static FString PurgePageName(const FString& NewPageName);
+
+	UFUNCTION(BlueprintCallable, Category="Render Pages | Page")
 	void SetPageName(const FString& NewPageName) { PageName = NewPageName; }
 	void SetPageNameRaw(const FString& NewPageName) { PageName = NewPageName; }
 
+
+	UFUNCTION(BlueprintPure, Category="Render Pages | Page")
 	bool GetIsEnabled() const { return bIsEnabled; }
+
+	UFUNCTION(BlueprintCallable, Category="Render Pages | Page")
 	void SetIsEnabled(const bool bEnabled) { bIsEnabled = bEnabled; }
 
+
+	UFUNCTION(BlueprintPure, Category="Render Pages | Page")
 	FString GetOutputDirectory() const;
 	FString GetOutputDirectoryRaw() const { return OutputDirectory; }
+
+	UFUNCTION(BlueprintPure, Category="Render Pages | Page")
 	FString GetOutputDirectoryForDisplay() const { return OutputDirectory; }
+
 	static FString PurgeOutputDirectory(const FString& NewOutputDirectory);
+
+	UFUNCTION(BlueprintCallable, Category="Render Pages | Page")
 	void SetOutputDirectory(const FString& NewOutputDirectory) { OutputDirectory = PurgeOutputDirectory(NewOutputDirectory); }
 	void SetOutputDirectoryRaw(const FString& NewOutputDirectory) { OutputDirectory = NewOutputDirectory; }
 
+
+	UFUNCTION(BlueprintPure, Category="Render Pages | Page")
 	UMoviePipelineMasterConfig* GetRenderPreset() const { return RenderPreset; }
+
+	UFUNCTION(BlueprintPure, Category="Render Pages | Page")
 	UMoviePipelineOutputSetting* GetRenderPresetOutputSettings() const;
+
+	UFUNCTION(BlueprintCallable, Category="Render Pages | Page")
 	void SetRenderPreset(UMoviePipelineMasterConfig* NewRenderPreset) { RenderPreset = NewRenderPreset; }
 
+
 	bool HasRemoteControlValue(const TSharedPtr<FRemoteControlEntity>& RemoteControlEntity) const;
 	bool ConstGetRemoteControlValue(const TSharedPtr<FRemoteControlEntity>& RemoteControlEntity, TArray<uint8>& OutBinaryArray) const;
 	bool GetRemoteControlValue(const TSharedPtr<FRemoteControlEntity>& RemoteControlEntity, TArray<uint8>& OutBinaryArray);
@@ -147,8 +212,8 @@ private:
 	FGuid Id;
 
 	/** Waits the given number of frames before it will render this page. This can be set to a higher amount when the renderer has to wait for your code to complete (such as construction scripts etc). Try increasing this value when rendering doesn't produce the output you expect it to. */
-	UPROPERTY(EditInstanceOnly, Category="Render Pages|Page", Meta=(AllowPrivateAccess="true"))
-	uint32 WaitFramesBeforeRendering;
+	UPROPERTY(EditInstanceOnly, Category="Render Pages|Page", Meta=(AllowPrivateAccess="true", ClampMin="0"))
+	int32 WaitFramesBeforeRendering;
 
 	/** The level sequence, this is what will be rendered during rendering. A render page without a level sequence can't be rendered. */
 	UPROPERTY(EditInstanceOnly, Category="Render Pages|Page", Meta=(AllowPrivateAccess="true"))
@@ -225,6 +290,25 @@ public:
 	/** Should be called when the editor closes this asset. */
 	void OnClose() { SaveValuesToCDO(); }
 
+public:
+	static TArray<FString> GetBlueprintImplementableEvents() { return {TEXT("ReceivePreRender"), TEXT("ReceivePostRender")}; }
+
+protected:
+	/** Event for when rendering begins for a page. */
+	UFUNCTION(BlueprintImplementableEvent, Meta=(DisplayName="PreRender"))
+	void ReceivePreRender(URenderPage* Page);
+
+	/** Event for when rendering ends for a page. */
+	UFUNCTION(BlueprintImplementableEvent, Meta=(DisplayName="PostRender"))
+	void ReceivePostRender(URenderPage* Page);
+
+public:
+	/** Overridable native event for when rendering begins for a page. */
+	virtual void PreRender(URenderPage* Page);
+
+	/** Overridable native event for when rendering ends for a page. */
+	virtual void PostRender(URenderPage* Page);
+
 private:
 	/**
 	 * Because render page collection assets are blueprints (assets that also have a blueprint graph), the render page collection data is not stored directly in the asset data that you see in the content browser.
@@ -234,10 +318,13 @@ private:
 
 	/** Obtains the CDO, could return itself if this is called on the CDO instance. */
 	URenderPageCollection* GetCDO() { return (HasAnyFlags(RF_ClassDefaultObject) ? this : GetClass()->GetDefaultObject<URenderPageCollection>()); }
+
 	/** Copied values over into the CDO. */
 	void SaveValuesToCDO() { CopyValuesToOrFromCDO(true); }
+
 	/** Copied values over from the CDO. */
 	void LoadValuesFromCDO() { CopyValuesToOrFromCDO(false); }
+
 	/** Copied values to or from the CDO, based on whether bToCDO is true or false. */
 	void CopyValuesToOrFromCDO(const bool bToCDO);
 
diff --git a/Engine/Plugins/Experimental/RenderPages/Source/RenderPages/Public/RenderPage/RenderPageManager.h b/Engine/Plugins/Experimental/RenderPages/Source/RenderPages/Public/RenderPage/RenderPageManager.h
--- a/Engine/Plugins/Experimental/RenderPages/Source/RenderPages/Public/RenderPage/RenderPageManager.h
+++ b/Engine/Plugins/Experimental/RenderPages/Source/RenderPages/Public/RenderPage/RenderPageManager.h
@@ -28,6 +28,9 @@ public:
 		: RemoteControlData(InRemoteControlData)
 	{}
 
+public:
+	bool IsEmpty() const { return RemoteControlData.IsEmpty(); }
+
 public:
 	/** The previous values of the remote control properties. */
 	UPROPERTY()
diff --git a/Engine/Plugins/Experimental/RenderPages/Source/RenderPages/Public/RenderPage/RenderPageMoviePipelineJob.h b/Engine/Plugins/Experimental/RenderPages/Source/RenderPages/Public/RenderPage/RenderPageMoviePipelineJob.h
--- a/Engine/Plugins/Experimental/RenderPages/Source/RenderPages/Public/RenderPage/RenderPageMoviePipelineJob.h
+++ b/Engine/Plugins/Experimental/RenderPages/Source/RenderPages/Public/RenderPage/RenderPageMoviePipelineJob.h
@@ -94,6 +94,10 @@ public:
 	/** Retrieves the "Engine Warm Up Count" value from the AntiAliasingSettings from the render preset that this render page uses. */
 	int32 GetEngineWarmUpCount() const;
 
+public:
+	/** Returns true if this render job was canceled (which for example can be caused by calling Cancel(), or by closing the render popup). */
+	bool IsCanceled() const { return bCanceled; }
+
 private:
 	void ComputePlaybackContext(bool& bOutAllowBinding);
 	void ExecuteFinished(UMoviePipelineExecutorBase* PipelineExecutor, const bool bSuccess);
@@ -124,6 +128,10 @@ protected:
 	/** Whether the entry can execute, or whether it should just skip execution. */
 	UPROPERTY(Transient)
 	bool bCanExecute;
+
+	/** Whether the entry was canceled (like by calling Cancel(), or by closing the render popup). */
+	UPROPERTY(Transient)
+	bool bCanceled;
 };
 
 
@@ -145,6 +153,9 @@ public:
 	/** Cancels this render job. Relies on the internal MRQ implementation of job canceling on whether this will stop the current page from rendering or not. Will always prevent new pages from rendering. */
 	void Cancel();
 
+	/** Returns true if this render job has been canceled. */
+	bool IsCanceled() const { return bCanceled; }
+
 	/** Retrieves the rendering status of the given render page. */
 	FString GetPageStatus(URenderPage* Page) const;
 
@@ -172,6 +183,10 @@ protected:
 	UPROPERTY(Transient)
 	FRenderPagePreviousEngineFpsSettings PreviousFrameLimitSettings;
 
+	/** True if the queue has previously executed the pre-render event of a page. */
+	UPROPERTY(Transient)
+	bool bRanPreRender;
+
 public:
 	/** A delegate for when the render job is about to start. */
 	UE::RenderPages::FOnRenderPagesMoviePipelineRenderJobStarted& OnExecuteStarted() { return OnExecuteStartedDelegate; }
diff --git a/Engine/Plugins/Experimental/RenderPages/Source/RenderPages/Public/Utils/RenderPageLevelSequencePlayer.h b/Engine/Plugins/Experimental/RenderPages/Source/RenderPages/Public/Utils/RenderPageLevelSequencePlayer.h
deleted file mode 100644
diff --git a/Engine/Plugins/Experimental/RenderPages/Source/RenderPages/RenderPages.Build.cs b/Engine/Plugins/Experimental/RenderPages/Source/RenderPages/RenderPages.Build.cs
--- a/Engine/Plugins/Experimental/RenderPages/Source/RenderPages/RenderPages.Build.cs
+++ b/Engine/Plugins/Experimental/RenderPages/Source/RenderPages/RenderPages.Build.cs
@@ -14,7 +14,6 @@ public class RenderPages : ModuleRules
 				"Core",
 				"CoreUObject",
 				"Engine",
-				"LevelSequence",
 				"RemoteControl",
 			}
 		);
diff --git a/Engine/Plugins/Experimental/RenderPages/Source/RenderPagesDeveloper/Private/Blueprints/RenderPagesBlueprint.cpp b/Engine/Plugins/Experimental/RenderPages/Source/RenderPagesDeveloper/Private/Blueprints/RenderPagesBlueprint.cpp
--- a/Engine/Plugins/Experimental/RenderPages/Source/RenderPagesDeveloper/Private/Blueprints/RenderPagesBlueprint.cpp
+++ b/Engine/Plugins/Experimental/RenderPages/Source/RenderPagesDeveloper/Private/Blueprints/RenderPagesBlueprint.cpp
@@ -1,18 +1,15 @@
 // Copyright Epic Games, Inc. All Rights Reserved.
 
 #include "Blueprints/RenderPagesBlueprint.h"
+
+#include "EdGraphSchema_K2_Actions.h"
 #include "Graph/RenderPagesGraph.h"
+#include "Kismet2/BlueprintEditorUtils.h"
+#include "Kismet2/KismetEditorUtilities.h"
 #include "RenderPage/RenderPageCollection.h"
 #include "RenderPage/RenderPagesBlueprintGeneratedClass.h"
 
 
-URenderPagesBlueprint::URenderPagesBlueprint()
-{
-	CompileLog.SetSourcePath(GetPathName());
-	CompileLog.bLogDetailedResults = false;
-	CompileLog.EventDisplayThresholdMs = false;
-}
-
 UClass* URenderPagesBlueprint::GetBlueprintClass() const
 {
 	return URenderPagesBlueprintGeneratedClass::StaticClass();
@@ -22,26 +19,45 @@ void URenderPagesBlueprint::PostLoad()
 {
 	Super::PostLoad();
 
-	{
-		// Remove all non Render Pages graphs
-		TArray<UEdGraph*> NewUberGraphPages;
+	{// remove every URenderPagesGraph (because that class is deprecated) >>
+		bool bChanged = false;
+		TArray<TObjectPtr<UEdGraph>> NewUberGraphPages;
 
-		for (UEdGraph* Graph : UbergraphPages)
+		for (const TObjectPtr<UEdGraph>& Graph : UbergraphPages)
 		{
-			if (URenderPagesGraph* RenderPagesGraph = Cast<URenderPagesGraph>(Graph))
+			if (UDEPRECATED_RenderPagesGraph* RenderPagesGraph = Cast<UDEPRECATED_RenderPagesGraph>(Graph))
 			{
-				NewUberGraphPages.Add(RenderPagesGraph);
+				bChanged = true;
+				RenderPagesGraph->MarkAsGarbage();
+				RenderPagesGraph->Rename(nullptr, GetTransientPackage(), REN_ForceNoResetLoaders);
+				continue;
 			}
-			else
+			NewUberGraphPages.Add(Graph);
+		}
+
+		if (bChanged)
+		{
+			UbergraphPages = NewUberGraphPages;
+		}
+	}// remove every URenderPagesGraph (because that class is deprecated) <<
+
+	if (UbergraphPages.IsEmpty())
+	{
+		UEdGraph* NewGraph = FBlueprintEditorUtils::CreateNewGraph(this, UEdGraphSchema_K2::GN_EventGraph, UEdGraph::StaticClass(), UEdGraphSchema_K2::StaticClass());
+		NewGraph->bAllowDeletion = false;
+
+		{// create every RenderPages blueprint event >>
+			int32 i = 0;
+			for (const FString& Event : URenderPageCollection::GetBlueprintImplementableEvents())
 			{
-				Graph->MarkAsGarbage();
-				Graph->Rename(nullptr, GetTransientPackage(), REN_ForceNoResetLoaders);
+				int32 InOutNodePosY = (i * 256) - 48;
+				FKismetEditorUtilities::AddDefaultEventNode(this, NewGraph, FName(Event), URenderPageCollection::StaticClass(), InOutNodePosY);
+				i++;
 			}
-		}
+		}// create every RenderPages blueprint event <<
 
-		UbergraphPages = NewUberGraphPages;
+		FBlueprintEditorUtils::MarkBlueprintAsStructurallyModified(this);
+		FBlueprintEditorUtils::AddUbergraphPage(this, NewGraph);
+		LastEditedDocuments.AddUnique(NewGraph);
 	}
-
-	CompileLog.Messages.Reset();
-	CompileLog.NumErrors = CompileLog.NumWarnings = 0;
 }
diff --git a/Engine/Plugins/Experimental/RenderPages/Source/RenderPagesDeveloper/Private/Blueprints/RenderPagesBlueprintCompiler.cpp b/Engine/Plugins/Experimental/RenderPages/Source/RenderPagesDeveloper/Private/Blueprints/RenderPagesBlueprintCompiler.cpp
--- a/Engine/Plugins/Experimental/RenderPages/Source/RenderPagesDeveloper/Private/Blueprints/RenderPagesBlueprintCompiler.cpp
+++ b/Engine/Plugins/Experimental/RenderPages/Source/RenderPagesDeveloper/Private/Blueprints/RenderPagesBlueprintCompiler.cpp
@@ -2,10 +2,26 @@
 
 #include "Blueprints/RenderPagesBlueprintCompiler.h"
 
+#include "RenderPage/RenderPageCollection.h"
+#include "Stats/StatsHierarchical.h"
+
+#define LOCTEXT_NAMESPACE "RenderPagesBlueprintCompiler"
+
 
 bool UE::RenderPages::FRenderPagesBlueprintCompiler::CanCompile(const UBlueprint* Blueprint)
 {
-	return false;
+	DECLARE_SCOPE_HIERARCHICAL_COUNTER_FUNC()
+
+	return (Blueprint && Blueprint->ParentClass && Blueprint->ParentClass->IsChildOf(URenderPageCollection::StaticClass()));
+}
+
+void UE::RenderPages::FRenderPagesBlueprintCompiler::Compile(UBlueprint* Blueprint, const FKismetCompilerOptions& CompileOptions, FCompilerResultsLog& Results)
+{
+	DECLARE_SCOPE_HIERARCHICAL_COUNTER_FUNC()
+
+	FRenderPagesBlueprintCompilerContext Compiler(Blueprint, Results, CompileOptions);
+	Compiler.Compile();
 }
 
-void UE::RenderPages::FRenderPagesBlueprintCompiler::Compile(UBlueprint* Blueprint, const FKismetCompilerOptions& CompileOptions, FCompilerResultsLog& Results) {}
+
+#undef LOCTEXT_NAMESPACE
diff --git a/Engine/Plugins/Experimental/RenderPages/Source/RenderPagesDeveloper/Private/Graph/RenderPagesGraph.cpp b/Engine/Plugins/Experimental/RenderPages/Source/RenderPagesDeveloper/Private/Graph/RenderPagesGraph.cpp
--- a/Engine/Plugins/Experimental/RenderPages/Source/RenderPagesDeveloper/Private/Graph/RenderPagesGraph.cpp
+++ b/Engine/Plugins/Experimental/RenderPages/Source/RenderPagesDeveloper/Private/Graph/RenderPagesGraph.cpp
@@ -8,17 +8,17 @@
 #define LOCTEXT_NAMESPACE "RenderPagesGraph"
 
 
-void URenderPagesGraph::Initialize(URenderPagesBlueprint* InBlueprint) {}
+void UDEPRECATED_RenderPagesGraph::Initialize(URenderPagesBlueprint* InBlueprint) {}
 
-void URenderPagesGraph::PostLoad()
+void UDEPRECATED_RenderPagesGraph::PostLoad()
 {
 	Super::PostLoad();
-	Schema = URenderPagesGraphSchema::StaticClass();
+	Schema = UDEPRECATED_RenderPagesGraphSchema::StaticClass();
 }
 
-URenderPagesBlueprint* URenderPagesGraph::GetBlueprint() const
+URenderPagesBlueprint* UDEPRECATED_RenderPagesGraph::GetBlueprint() const
 {
-	if (URenderPagesGraph* OuterGraph = Cast<URenderPagesGraph>(GetOuter()))
+	if (UDEPRECATED_RenderPagesGraph* OuterGraph = Cast<UDEPRECATED_RenderPagesGraph>(GetOuter()))
 	{
 		return OuterGraph->GetBlueprint();
 	}
diff --git a/Engine/Plugins/Experimental/RenderPages/Source/RenderPagesDeveloper/Private/Graph/RenderPagesGraphNode.cpp b/Engine/Plugins/Experimental/RenderPages/Source/RenderPagesDeveloper/Private/Graph/RenderPagesGraphNode.cpp
deleted file mode 100644
diff --git a/Engine/Plugins/Experimental/RenderPages/Source/RenderPagesDeveloper/Private/Graph/RenderPagesGraphSchema.cpp b/Engine/Plugins/Experimental/RenderPages/Source/RenderPagesDeveloper/Private/Graph/RenderPagesGraphSchema.cpp
--- a/Engine/Plugins/Experimental/RenderPages/Source/RenderPagesDeveloper/Private/Graph/RenderPagesGraphSchema.cpp
+++ b/Engine/Plugins/Experimental/RenderPages/Source/RenderPagesDeveloper/Private/Graph/RenderPagesGraphSchema.cpp
@@ -1,255 +1,3 @@
 // Copyright Epic Games, Inc. All Rights Reserved.
 
 #include "Graph/RenderPagesGraphSchema.h"
-#include "Graph/RenderPagesGraph.h"
-#include "Graph/RenderPagesGraphNode.h"
-#include "Blueprints/RenderPagesBlueprint.h"
-#include "EdGraphSchema_K2_Actions.h"
-#include "Kismet2/BlueprintEditorUtils.h"
-#include "Kismet2/Kismet2NameValidators.h"
-
-#define LOCTEXT_NAMESPACE "RenderPagesGraphSchema"
-
-/////////////////////////////////////////////////////
-// UE::RenderPages::FRenderPagesLocalVariableNameValidator
-UE::RenderPages::FRenderPagesLocalVariableNameValidator::FRenderPagesLocalVariableNameValidator(const UBlueprint* Blueprint, const URenderPagesGraph* Graph, FName InExistingName)
-	: FStringSetNameValidator(InExistingName.ToString())
-{
-	if (Blueprint)
-	{
-		TSet<FName> NamesTemp;
-		// We allow local variables with same name as blueprint variable
-
-		FBlueprintEditorUtils::GetFunctionNameList(Blueprint, NamesTemp);
-		FBlueprintEditorUtils::GetAllGraphNames(Blueprint, NamesTemp);
-		FBlueprintEditorUtils::GetSCSVariableNameList(Blueprint, NamesTemp);
-		FBlueprintEditorUtils::GetImplementingBlueprintsFunctionNameList(Blueprint, NamesTemp);
-
-		for (FName& Name : NamesTemp)
-		{
-			Names.Add(Name.ToString());
-		}
-	}
-}
-
-/////////////////////////////////////////////////////
-// UE::RenderPages::FRenderPagesNameValidator
-UE::RenderPages::FRenderPagesNameValidator::FRenderPagesNameValidator(const UBlueprint* Blueprint, const UStruct* ValidationScope, FName InExistingName)
-	: FStringSetNameValidator(InExistingName.ToString())
-{
-	if (Blueprint)
-	{
-		TSet<FName> NamesTemp;
-		FBlueprintEditorUtils::GetClassVariableList(Blueprint, NamesTemp, true);
-		FBlueprintEditorUtils::GetFunctionNameList(Blueprint, NamesTemp);
-		FBlueprintEditorUtils::GetAllGraphNames(Blueprint, NamesTemp);
-		FBlueprintEditorUtils::GetSCSVariableNameList(Blueprint, NamesTemp);
-		FBlueprintEditorUtils::GetImplementingBlueprintsFunctionNameList(Blueprint, NamesTemp);
-
-		for (FName& Name : NamesTemp)
-		{
-			Names.Add(Name.ToString());
-		}
-	}
-}
-
-/////////////////////////////////////////////////////
-// FRenderPagesGraphSchemaAction_LocalVar
-bool FRenderPagesGraphSchemaAction_LocalVar::IsValidName(const FName& NewName, FText& OutErrorMessage) const
-{
-	if (URenderPagesGraph* Graph = Cast<URenderPagesGraph>(GetVariableScope()))
-	{
-		UE::RenderPages::FRenderPagesLocalVariableNameValidator NameValidator(Graph->GetBlueprint(), Graph, GetVariableName());
-		EValidatorResult Result = NameValidator.IsValid(NewName.ToString(), false);
-		if (Result != EValidatorResult::Ok && Result != EValidatorResult::ExistingName)
-		{
-			OutErrorMessage = FText::FromString(TEXT("Name with invalid format"));
-			return false;
-		}
-	}
-	return FEdGraphSchemaAction_BlueprintVariableBase::IsValidName(NewName, OutErrorMessage);
-}
-
-/////////////////////////////////////////////////////
-// FRenderPagesGraphSchemaAction_PromoteToVariable
-FRenderPagesGraphSchemaAction_PromoteToVariable::FRenderPagesGraphSchemaAction_PromoteToVariable(UEdGraphPin* InEdGraphPin, bool InLocalVariable)
-	: FEdGraphSchemaAction(FText(),
-		InLocalVariable ? LOCTEXT("PromoteToLocalVariable", "Promote to local variable") : LOCTEXT("PromoteToVariable", "Promote to variable"),
-		InLocalVariable ? LOCTEXT("PromoteToLocalVariable", "Promote to local variable") : LOCTEXT("PromoteToVariable", "Promote to variable"),
-		1)
-	, EdGraphPin(InEdGraphPin)
-	, bLocalVariable(InLocalVariable)
-{}
-
-UEdGraphNode* FRenderPagesGraphSchemaAction_PromoteToVariable::PerformAction(UEdGraph* ParentGraph, UEdGraphPin* FromPin, const FVector2D Location, bool bSelectNewNode)
-{
-	return nullptr;
-}
-
-/////////////////////////////////////////////////////
-// UE::RenderPages::FRenderPagesFunctionDragDropAction
-UE::RenderPages::FRenderPagesFunctionDragDropAction::FRenderPagesFunctionDragDropAction()
-	: FGraphSchemaActionDragDropAction()
-	, SourceRenderPagesBlueprint(nullptr)
-	, SourceRenderPagesGraph(nullptr)
-	, bControlDrag(false)
-	, bAltDrag(false)
-{}
-
-TSharedRef<UE::RenderPages::FRenderPagesFunctionDragDropAction> UE::RenderPages::FRenderPagesFunctionDragDropAction::New(TSharedPtr<FEdGraphSchemaAction> InAction, URenderPagesBlueprint* InRenderPagesBlueprint, URenderPagesGraph* InRenderPagesGraph)
-{
-	TSharedRef<FRenderPagesFunctionDragDropAction> Action = MakeShareable(new FRenderPagesFunctionDragDropAction);
-	Action->SourceAction = InAction;
-	Action->SourceRenderPagesBlueprint = InRenderPagesBlueprint;
-	Action->SourceRenderPagesGraph = InRenderPagesGraph;
-	Action->Construct();
-	return Action;
-}
-
-/////////////////////////////////////////////////////
-// URenderPagesGraphSchema
-
-const FName URenderPagesGraphSchema::GraphName_RenderPages(TEXT("Logic Graph"));
-
-URenderPagesGraphSchema::URenderPagesGraphSchema() {}
-
-void URenderPagesGraphSchema::GetGraphDisplayInformation(const UEdGraph& Graph, FGraphDisplayInfo& DisplayInfo) const
-{
-	Super::GetGraphDisplayInformation(Graph, DisplayInfo);
-
-	if (Cast<URenderPagesGraph>(const_cast<UEdGraph*>(&Graph)))
-	{
-		static const FText MainGraphText = FText::FromString(TEXT("The logic graph for the Render Pages."));
-		DisplayInfo.Tooltip = MainGraphText;
-	}
-}
-
-FLinearColor URenderPagesGraphSchema::GetPinTypeColor(const FEdGraphPinType& PinType) const
-{
-	return GetDefault<UEdGraphSchema_K2>()->GetPinTypeColor(PinType);
-}
-
-bool URenderPagesGraphSchema::CanGraphBeDropped(TSharedPtr<FEdGraphSchemaAction> InAction) const
-{
-	if (!InAction.IsValid())
-	{
-		return false;
-	}
-
-	if (InAction->GetTypeId() == FEdGraphSchemaAction_K2Graph::StaticGetTypeId())
-	{
-		FEdGraphSchemaAction_K2Graph* FuncAction = static_cast<FEdGraphSchemaAction_K2Graph*>(InAction.Get());
-		if (Cast<URenderPagesGraph>(FuncAction->EdGraph))
-		{
-			return true;
-		}
-	}
-	else if (InAction->GetTypeId() == FRenderPagesGraphSchemaAction_LocalVar::StaticGetTypeId())
-	{
-		FRenderPagesGraphSchemaAction_LocalVar* VarAction = static_cast<FRenderPagesGraphSchemaAction_LocalVar*>(InAction.Get());
-		if (Cast<URenderPagesGraph>(VarAction->GetVariableScope()))
-		{
-			return true;
-		}
-	}
-
-	return false;
-}
-
-FReply URenderPagesGraphSchema::BeginGraphDragAction(TSharedPtr<FEdGraphSchemaAction> InAction, const FPointerEvent& MouseEvent) const
-{
-	if (!InAction.IsValid())
-	{
-		return FReply::Unhandled();
-	}
-
-	if (InAction->GetTypeId() == FEdGraphSchemaAction_K2Graph::StaticGetTypeId())
-	{
-		FEdGraphSchemaAction_K2Graph* FuncAction = static_cast<FEdGraphSchemaAction_K2Graph*>(InAction.Get());
-		if (URenderPagesGraph* RenderPagesGraph = Cast<URenderPagesGraph>(FuncAction->EdGraph))
-		{
-			if (URenderPagesBlueprint* RenderPagesBlueprint = Cast<URenderPagesBlueprint>(FBlueprintEditorUtils::FindBlueprintForGraph(RenderPagesGraph)))
-			{
-				TSharedRef<UE::RenderPages::FRenderPagesFunctionDragDropAction> Action = UE::RenderPages::FRenderPagesFunctionDragDropAction::New(InAction, RenderPagesBlueprint, RenderPagesGraph);
-				Action->SetAltDrag(MouseEvent.IsAltDown());
-				Action->SetCtrlDrag(MouseEvent.IsControlDown());
-				return FReply::Handled().BeginDragDrop(Action);
-			}
-		}
-	}
-	else if (InAction->GetTypeId() == FRenderPagesGraphSchemaAction_LocalVar::StaticGetTypeId())
-	{
-		FRenderPagesGraphSchemaAction_LocalVar* VarAction = static_cast<FRenderPagesGraphSchemaAction_LocalVar*>(InAction.Get());
-		if (URenderPagesGraph* RenderPagesGraph = Cast<URenderPagesGraph>(VarAction->GetVariableScope()))
-		{
-			if (URenderPagesBlueprint* RenderPagesBlueprint = Cast<URenderPagesBlueprint>(FBlueprintEditorUtils::FindBlueprintForGraph(RenderPagesGraph)))
-			{
-				TSharedRef<UE::RenderPages::FRenderPagesFunctionDragDropAction> Action = UE::RenderPages::FRenderPagesFunctionDragDropAction::New(InAction, RenderPagesBlueprint, RenderPagesGraph);
-				Action->SetAltDrag(MouseEvent.IsAltDown());
-				Action->SetCtrlDrag(MouseEvent.IsControlDown());
-				return FReply::Handled().BeginDragDrop(Action);
-			}
-		}
-	}
-	return FReply::Unhandled();
-}
-
-bool URenderPagesGraphSchema::CanVariableBeDropped(UEdGraph* InGraph, FProperty* InVariableToDrop) const
-{
-	return true;
-}
-
-bool URenderPagesGraphSchema::RequestVariableDropOnPanel(UEdGraph* InGraph, FProperty* InVariableToDrop, const FVector2D& InDropPosition, const FVector2D& InScreenPosition)
-{
-	if (CanVariableBeDropped(InGraph, InVariableToDrop))
-	{
-		UBlueprint* Blueprint = FBlueprintEditorUtils::FindBlueprintForGraph(InGraph);
-		if (Cast<URenderPagesBlueprint>(Blueprint))
-		{
-			return true;
-		}
-	}
-	return false;
-}
-
-bool URenderPagesGraphSchema::RequestVariableDropOnPin(UEdGraph* InGraph, FProperty* InVariableToDrop, UEdGraphPin* InPin, const FVector2D& InDropPosition, const FVector2D& InScreenPosition)
-{
-	if (CanVariableBeDropped(InGraph, InVariableToDrop))
-	{
-		if (Cast<URenderPagesGraph>(InGraph))
-		{
-			return true;
-		}
-	}
-	return false;
-}
-
-void URenderPagesGraphSchema::InsertAdditionalActions(TArray<UBlueprint*> InBlueprints, TArray<UEdGraph*> InGraphs, TArray<UEdGraphPin*> InPins, FGraphActionListBuilderBase& OutAllActions) const
-{
-	Super::InsertAdditionalActions(InBlueprints, InGraphs, InPins, OutAllActions);
-
-	if (InPins.Num() > 0)
-	{
-		if (URenderPagesGraphNode* RenderPagesNode = Cast<URenderPagesGraphNode>(InPins[0]->GetOwningNode()))
-		{
-			if (UEdGraphPin* ModelPin = RenderPagesNode->GetPinAt(0))
-			{
-				OutAllActions.AddAction(MakeShared<FRenderPagesGraphSchemaAction_PromoteToVariable>(ModelPin, false));
-			}
-		}
-	}
-}
-
-TSharedPtr<INameValidatorInterface> URenderPagesGraphSchema::GetNameValidator(const UBlueprint* BlueprintObj, const FName& OriginalName, const UStruct* ValidationScope, const FName& ActionTypeId) const
-{
-	if (ActionTypeId == FRenderPagesGraphSchemaAction_LocalVar::StaticGetTypeId())
-	{
-		if (const URenderPagesGraph* RenderPagesGraph = Cast<URenderPagesGraph>(ValidationScope))
-		{
-			return MakeShareable(new UE::RenderPages::FRenderPagesLocalVariableNameValidator(BlueprintObj, RenderPagesGraph, OriginalName));
-		}
-	}
-	return MakeShareable(new UE::RenderPages::FRenderPagesNameValidator(BlueprintObj, ValidationScope, OriginalName));
-}
-
-#undef LOCTEXT_NAMESPACE
diff --git a/Engine/Plugins/Experimental/RenderPages/Source/RenderPagesDeveloper/Public/Blueprints/RenderPagesBlueprint.h b/Engine/Plugins/Experimental/RenderPages/Source/RenderPagesDeveloper/Public/Blueprints/RenderPagesBlueprint.h
--- a/Engine/Plugins/Experimental/RenderPages/Source/RenderPagesDeveloper/Public/Blueprints/RenderPagesBlueprint.h
+++ b/Engine/Plugins/Experimental/RenderPages/Source/RenderPagesDeveloper/Public/Blueprints/RenderPagesBlueprint.h
@@ -4,15 +4,9 @@
 
 #include "CoreMinimal.h"
 #include "Engine/Blueprint.h"
-#include "RenderPage/RenderPagePropsSource.h"
-#include "Kismet2/CompilerResultsLog.h"
 #include "RenderPagesBlueprint.generated.h"
 
 
-class URenderPagesBlueprint;
-class URenderPageCollection;
-
-
 /**
  * A UBlueprint child class for the RenderPages modules.
  *
@@ -23,25 +17,19 @@ class RENDERPAGESDEVELOPER_API URenderPagesBlueprint : public UBlueprint
 {
 	GENERATED_BODY()
 
-public:
-	URenderPagesBlueprint();
-
 public:
 	//~ Begin UBlueprint Interface
 	virtual UClass* GetBlueprintClass() const override;
 	virtual bool SupportedByDefaultBlueprintFactory() const override { return false; }
 	virtual bool IsValidForBytecodeOnlyRecompile() const override { return false; }
 	virtual bool SupportsGlobalVariables() const override { return true; }
-	virtual bool SupportsLocalVariables() const override { return false; }
+	virtual bool SupportsLocalVariables() const override { return true; }
 	virtual bool SupportsFunctions() const override { return true; }
-	virtual bool SupportsMacros() const override { return false; }
+	virtual bool SupportsMacros() const override { return true; }
 	virtual bool SupportsDelegates() const override { return true; }
-	virtual bool SupportsEventGraphs() const override { return false; }
+	virtual bool SupportsEventGraphs() const override { return true; }
 	virtual bool SupportsAnimLayers() const override { return false; }
 
 	virtual void PostLoad() override;
 	//~ End UBlueprint Interface
-
-private:
-	FCompilerResultsLog CompileLog;
 };
diff --git a/Engine/Plugins/Experimental/RenderPages/Source/RenderPagesDeveloper/Public/Blueprints/RenderPagesBlueprintCompiler.h b/Engine/Plugins/Experimental/RenderPages/Source/RenderPagesDeveloper/Public/Blueprints/RenderPagesBlueprintCompiler.h
--- a/Engine/Plugins/Experimental/RenderPages/Source/RenderPagesDeveloper/Public/Blueprints/RenderPagesBlueprintCompiler.h
+++ b/Engine/Plugins/Experimental/RenderPages/Source/RenderPagesDeveloper/Public/Blueprints/RenderPagesBlueprintCompiler.h
@@ -21,4 +21,13 @@ namespace UE::RenderPages
 		virtual void Compile(UBlueprint* Blueprint, const FKismetCompilerOptions& CompileOptions, FCompilerResultsLog& Results) override;
 		//~ End IBlueprintCompiler Interface
 	};
+
+
+	class RENDERPAGESDEVELOPER_API FRenderPagesBlueprintCompilerContext : public FKismetCompilerContext
+	{
+	public:
+		FRenderPagesBlueprintCompilerContext(UBlueprint* SourceSketch, FCompilerResultsLog& InMessageLog, const FKismetCompilerOptions& InCompilerOptions)
+			: FKismetCompilerContext(SourceSketch, InMessageLog, InCompilerOptions)
+		{}
+	};
 }
diff --git a/Engine/Plugins/Experimental/RenderPages/Source/RenderPagesDeveloper/Public/Graph/RenderPagesGraph.h b/Engine/Plugins/Experimental/RenderPages/Source/RenderPagesDeveloper/Public/Graph/RenderPagesGraph.h
--- a/Engine/Plugins/Experimental/RenderPages/Source/RenderPagesDeveloper/Public/Graph/RenderPagesGraph.h
+++ b/Engine/Plugins/Experimental/RenderPages/Source/RenderPagesDeveloper/Public/Graph/RenderPagesGraph.h
@@ -8,8 +8,6 @@
 
 
 class URenderPagesBlueprint;
-class URenderPagesGraphSchema;
-class URenderPageCollection;
 
 
 /**
@@ -17,8 +15,8 @@ class URenderPageCollection;
  *
  * Required in order for a RenderPageCollection to be able to have a blueprint graph.
  */
-UCLASS()
-class RENDERPAGESDEVELOPER_API URenderPagesGraph : public UEdGraph
+UCLASS(Deprecated)
+class RENDERPAGESDEVELOPER_API UDEPRECATED_RenderPagesGraph : public UEdGraph
 {
 	GENERATED_BODY()
 
diff --git a/Engine/Plugins/Experimental/RenderPages/Source/RenderPagesDeveloper/Public/Graph/RenderPagesGraphNode.h b/Engine/Plugins/Experimental/RenderPages/Source/RenderPagesDeveloper/Public/Graph/RenderPagesGraphNode.h
deleted file mode 100644
diff --git a/Engine/Plugins/Experimental/RenderPages/Source/RenderPagesDeveloper/Public/Graph/RenderPagesGraphSchema.h b/Engine/Plugins/Experimental/RenderPages/Source/RenderPagesDeveloper/Public/Graph/RenderPagesGraphSchema.h
--- a/Engine/Plugins/Experimental/RenderPages/Source/RenderPagesDeveloper/Public/Graph/RenderPagesGraphSchema.h
+++ b/Engine/Plugins/Experimental/RenderPages/Source/RenderPagesDeveloper/Public/Graph/RenderPagesGraphSchema.h
@@ -4,175 +4,16 @@
 
 #include "CoreMinimal.h"
 #include "EdGraph/EdGraphSchema.h"
-#include "EdGraphSchema_K2_Actions.h"
-#include "GraphEditorDragDropAction.h"
 #include "RenderPagesGraphSchema.generated.h"
 
 
-class URenderPagesBlueprint;
-class URenderPagesGraph;
-
-
-namespace UE::RenderPages
-{
-	/**
-	 * A FStringSetNameValidator child class for the RenderPages modules, for validating local variable names.
-	 *
-	 * Required in order for a RenderPageCollection to be able to have a blueprint graph.
-	 */
-	class RENDERPAGESDEVELOPER_API FRenderPagesLocalVariableNameValidator : public FStringSetNameValidator
-	{
-	public:
-		FRenderPagesLocalVariableNameValidator(const UBlueprint* Blueprint, const URenderPagesGraph* Graph, FName InExistingName = NAME_None);
-	};
-
-	/**
-	 * A FStringSetNameValidator child class for the RenderPages modules, for validating names.
-	 *
-	 * Required in order for a RenderPageCollection to be able to have a blueprint graph.
-	 */
-	class RENDERPAGESDEVELOPER_API FRenderPagesNameValidator : public FStringSetNameValidator
-	{
-	public:
-		FRenderPagesNameValidator(const UBlueprint* Blueprint, const UStruct* ValidationScope, FName InExistingName = NAME_None);
-	};
-}
-
-
-/**
- * A FEdGraphSchemaAction_BlueprintVariableBase child class for the RenderPages modules, for local variables.
- *
- * Required in order for a RenderPageCollection to be able to have a blueprint graph.
- */
-USTRUCT()
-struct RENDERPAGESDEVELOPER_API FRenderPagesGraphSchemaAction_LocalVar : public FEdGraphSchemaAction_BlueprintVariableBase
-{
-	GENERATED_BODY()
-
-public:
-	/* Returns the type ID. Returns the same string every time. */
-	static FName StaticGetTypeId()
-	{
-		static FName Type("FRenderPagesGraphSchemaAction_LocalVar");
-		return Type;
-	}
-
-	FRenderPagesGraphSchemaAction_LocalVar() = default;
-	FRenderPagesGraphSchemaAction_LocalVar(const FText& InNodeCategory, const FText& InMenuDesc, const FText& InToolTip, const int32 InGrouping, const int32 InSectionID)
-		: FEdGraphSchemaAction_BlueprintVariableBase(InNodeCategory, InMenuDesc, InToolTip, InGrouping, InSectionID)
-	{}
-
-	//~ Begin FEdGraphSchemaAction interface
-	virtual FName GetTypeId() const override { return StaticGetTypeId(); }
-	virtual bool IsA(const FName& InType) const override { return (InType == GetTypeId()) || (InType == FEdGraphSchemaAction_BlueprintVariableBase::StaticGetTypeId()); }
-	virtual bool IsValidName(const FName& NewName, FText& OutErrorMessage) const override;
-	//~ End FEdGraphSchemaAction interface
-};
-
-
-/**
- * A FEdGraphSchemaAction child class for the RenderPages modules, for the promote-to-variable action.
- *
- * Required in order for a RenderPageCollection to be able to have a blueprint graph.
- */
-USTRUCT()
-struct RENDERPAGESDEVELOPER_API FRenderPagesGraphSchemaAction_PromoteToVariable : public FEdGraphSchemaAction
-{
-	GENERATED_BODY()
-
-public:
-	/* Returns the type ID. Returns the same string every time. */
-	static FName StaticGetTypeId()
-	{
-		static FName Type("FRenderPagesGraphSchemaAction_PromoteToVariable");
-		return Type;
-	}
-
-	FRenderPagesGraphSchemaAction_PromoteToVariable() = default;
-	FRenderPagesGraphSchemaAction_PromoteToVariable(UEdGraphPin* InEdGraphPin, bool InLocalVariable);
-
-	//~ Begin FEdGraphSchemaAction interface
-	virtual FName GetTypeId() const override { return StaticGetTypeId(); }
-	virtual bool IsA(const FName& InType) const override { return (InType == GetTypeId()); }
-	virtual UEdGraphNode* PerformAction(UEdGraph* ParentGraph, UEdGraphPin* FromPin, const FVector2D Location, bool bSelectNewNode) override;
-	//~ End FEdGraphSchemaAction interface
-
-private:
-	TObjectPtr<UEdGraphPin> EdGraphPin = nullptr;
-	bool bLocalVariable = false;
-};
-
-
-namespace UE::RenderPages
-{
-	/**
-	 * A FGraphSchemaActionDragDropAction child class for the RenderPages modules, for drag and dropping an item from the blueprints tree (like a variable or a function).
-	 *
-	 * Required in order for a RenderPageCollection to be able to have a blueprint graph.
-	 */
-	class RENDERPAGESDEVELOPER_API FRenderPagesFunctionDragDropAction : public FGraphSchemaActionDragDropAction
-	{
-	public:
-		DRAG_DROP_OPERATOR_TYPE(FRenderPagesFunctionDragDropAction, FGraphSchemaActionDragDropAction)
-
-	protected:
-		FRenderPagesFunctionDragDropAction();
-
-	public:
-		/** Set if operation is modified by alt. */
-		void SetAltDrag(const bool bInIsAltDrag) { bAltDrag = bInIsAltDrag; }
-
-		/** Set if operation is modified by the ctrl key. */
-		void SetCtrlDrag(const bool bInIsCtrlDrag) { bControlDrag = bInIsCtrlDrag; }
-
-	public:
-		/** Creates and returns a new instance of this class. */
-		static TSharedRef<FRenderPagesFunctionDragDropAction> New(TSharedPtr<FEdGraphSchemaAction> InAction, URenderPagesBlueprint* InRenderPagesBlueprint, URenderPagesGraph* InRenderPagesGraph);
-
-	protected:
-		TObjectPtr<URenderPagesBlueprint> SourceRenderPagesBlueprint;
-		TObjectPtr<URenderPagesGraph> SourceRenderPagesGraph;
-		bool bControlDrag;
-		bool bAltDrag;
-
-		friend class URenderPagesGraphSchema;
-	};
-}
-
-
 /**
  * A UEdGraphSchema child class for the RenderPages modules.
  *
  * Required in order for a RenderPageCollection to be able to have a blueprint graph.
  */
-UCLASS()
-class RENDERPAGESDEVELOPER_API URenderPagesGraphSchema : public UEdGraphSchema
+UCLASS(Deprecated)
+class RENDERPAGESDEVELOPER_API UDEPRECATED_RenderPagesGraphSchema : public UEdGraphSchema
 {
 	GENERATED_BODY()
-
-public:
-	URenderPagesGraphSchema();
-
-	//~ Begin UEdGraphSchema Interface
-	virtual void GetGraphDisplayInformation(const UEdGraph& Graph, /*out*/ FGraphDisplayInfo& DisplayInfo) const override;
-
-	virtual FLinearColor GetPinTypeColor(const FEdGraphPinType& PinType) const override;
-
-	virtual bool CanDuplicateGraph(UEdGraph* InSourceGraph) const override { return false; }
-
-	virtual bool CanGraphBeDropped(TSharedPtr<FEdGraphSchemaAction> InAction) const override;
-	virtual FReply BeginGraphDragAction(TSharedPtr<FEdGraphSchemaAction> InAction, const FPointerEvent& MouseEvent = FPointerEvent()) const override;
-
-	virtual bool CanVariableBeDropped(UEdGraph* InGraph, FProperty* InVariableToDrop) const override;
-	virtual bool RequestVariableDropOnPanel(UEdGraph* InGraph, FProperty* InVariableToDrop, const FVector2D& InDropPosition, const FVector2D& InScreenPosition) override;
-	virtual bool RequestVariableDropOnPin(UEdGraph* InGraph, FProperty* InVariableToDrop, UEdGraphPin* InPin, const FVector2D& InDropPosition, const FVector2D& InScreenPosition) override;
-
-	virtual void InsertAdditionalActions(TArray<UBlueprint*> InBlueprints, TArray<UEdGraph*> InGraphs, TArray<UEdGraphPin*> InPins, FGraphActionListBuilderBase& OutAllActions) const override;
-
-	virtual TSharedPtr<INameValidatorInterface> GetNameValidator(const UBlueprint* BlueprintObj, const FName& OriginalName, const UStruct* ValidationScope, const FName& ActionTypeId) const override;
-	//~ End UEdGraphSchema Interface
-
-public:
-	/** Name constant. */
-	static const FName GraphName_RenderPages;
 };
diff --git a/Engine/Plugins/Experimental/RenderPages/Source/RenderPagesDeveloper/RenderPagesDeveloper.Build.cs b/Engine/Plugins/Experimental/RenderPages/Source/RenderPagesDeveloper/RenderPagesDeveloper.Build.cs
--- a/Engine/Plugins/Experimental/RenderPages/Source/RenderPagesDeveloper/RenderPagesDeveloper.Build.cs
+++ b/Engine/Plugins/Experimental/RenderPages/Source/RenderPagesDeveloper/RenderPagesDeveloper.Build.cs
@@ -13,6 +13,7 @@ public class RenderPagesDeveloper : ModuleRules
 			{
 				"Core",
 				"CoreUObject",
+				"KismetCompiler",
 			}
 		);
 
diff --git a/Engine/Plugins/Experimental/RenderPages/Source/RenderPagesEditor/Private/BlueprintModes/RenderPagesApplicationModeListing.cpp b/Engine/Plugins/Experimental/RenderPages/Source/RenderPagesEditor/Private/BlueprintModes/RenderPagesApplicationModeListing.cpp
--- a/Engine/Plugins/Experimental/RenderPages/Source/RenderPagesEditor/Private/BlueprintModes/RenderPagesApplicationModeListing.cpp
+++ b/Engine/Plugins/Experimental/RenderPages/Source/RenderPagesEditor/Private/BlueprintModes/RenderPagesApplicationModeListing.cpp
@@ -101,12 +101,12 @@ void UE::RenderPages::Private::FRenderPagesApplicationModeListing::RegisterTabFa
 
 void UE::RenderPages::Private::FRenderPagesApplicationModeListing::PreDeactivateMode()
 {
-	// prevented from running: FRenderPagesBlueprintApplicationMode::PreDeactivateMode();
+	// prevents: FRenderPagesApplicationModeBase::PreDeactivateMode();
 }
 
 void UE::RenderPages::Private::FRenderPagesApplicationModeListing::PostActivateMode()
 {
-	// prevented from running: FRenderPagesBlueprintApplicationMode::PostActivateMode();
+	// prevents: FRenderPagesApplicationModeBase::PostActivateMode();
 }
 
 
diff --git a/Engine/Plugins/Experimental/RenderPages/Source/RenderPagesEditor/Private/BlueprintModes/RenderPagesApplicationModeLogic.cpp b/Engine/Plugins/Experimental/RenderPages/Source/RenderPagesEditor/Private/BlueprintModes/RenderPagesApplicationModeLogic.cpp
--- a/Engine/Plugins/Experimental/RenderPages/Source/RenderPagesEditor/Private/BlueprintModes/RenderPagesApplicationModeLogic.cpp
+++ b/Engine/Plugins/Experimental/RenderPages/Source/RenderPagesEditor/Private/BlueprintModes/RenderPagesApplicationModeLogic.cpp
@@ -2,15 +2,10 @@
 
 #include "BlueprintModes/RenderPagesApplicationModeLogic.h"
 #include "BlueprintModes/RenderPagesApplicationModes.h"
-
-#include "BlueprintEditorSharedTabFactories.h"
-#include "BlueprintEditorTabs.h"
-
-#include "IRenderPagesEditorModule.h"
 #include "IRenderPageCollectionEditor.h"
 
+#include "BlueprintEditorTabs.h"
 #include "SBlueprintEditorToolbar.h"
-
 #include "Toolkit/RenderPageCollectionEditorToolbar.h"
 
 #define LOCTEXT_NAMESPACE "RenderPagesLogicMode"
@@ -76,7 +71,7 @@ UE::RenderPages::Private::FRenderPagesApplicationModeLogic::FRenderPagesApplicat
 
 		InRenderPagesEditor->GetToolbarBuilder()->AddCompileToolbar(Toolbar);
 		InRenderPagesEditor->GetToolbarBuilder()->AddScriptingToolbar(Toolbar);
-		InRenderPagesEditor->GetToolbarBuilder()->AddBlueprintGlobalOptionsToolbar(Toolbar);
+		// disabled: InRenderPagesEditor->GetToolbarBuilder()->AddBlueprintGlobalOptionsToolbar(Toolbar);
 	}
 }
 
@@ -92,7 +87,7 @@ void UE::RenderPages::Private::FRenderPagesApplicationModeLogic::RegisterTabFact
 
 void UE::RenderPages::Private::FRenderPagesApplicationModeLogic::PreDeactivateMode()
 {
-	// prevented from running: FRenderPagesBlueprintApplicationMode::PreDeactivateMode();
+	// prevents: FRenderPagesApplicationModeBase::PreDeactivateMode();
 }
 
 void UE::RenderPages::Private::FRenderPagesApplicationModeLogic::PostActivateMode()
diff --git a/Engine/Plugins/Experimental/RenderPages/Source/RenderPagesEditor/Private/Commands/RenderPagesEditorCommands.cpp b/Engine/Plugins/Experimental/RenderPages/Source/RenderPagesEditor/Private/Commands/RenderPagesEditorCommands.cpp
--- a/Engine/Plugins/Experimental/RenderPages/Source/RenderPagesEditor/Private/Commands/RenderPagesEditorCommands.cpp
+++ b/Engine/Plugins/Experimental/RenderPages/Source/RenderPagesEditor/Private/Commands/RenderPagesEditorCommands.cpp
@@ -12,8 +12,6 @@ void UE::RenderPages::Private::FRenderPagesEditorCommands::RegisterCommands()
 	UI_COMMAND(DeletePage, "Delete", "Deletes an existing page instance from the list.", EUserInterfaceActionType::Button, FInputChord());
 
 	UI_COMMAND(BatchRenderList, "Render", "Renders the available page instance(s) in batch.", EUserInterfaceActionType::Button, FInputChord());
-
-	UI_COMMAND(AutoCompileGraph, "Auto Compile", "Auto-compile the render pages graph if turned on.", EUserInterfaceActionType::ToggleButton, FInputChord());
 }
 
 
diff --git a/Engine/Plugins/Experimental/RenderPages/Source/RenderPagesEditor/Private/Commands/RenderPagesEditorCommands.h b/Engine/Plugins/Experimental/RenderPages/Source/RenderPagesEditor/Private/Commands/RenderPagesEditorCommands.h
--- a/Engine/Plugins/Experimental/RenderPages/Source/RenderPagesEditor/Private/Commands/RenderPagesEditorCommands.h
+++ b/Engine/Plugins/Experimental/RenderPages/Source/RenderPagesEditor/Private/Commands/RenderPagesEditorCommands.h
@@ -36,8 +36,5 @@ namespace UE::RenderPages::Private
 
 		/** Renders the available page(s) in batch. */
 		TSharedPtr<FUICommandInfo> BatchRenderList;
-
-		/** Toggles Auto Compilation in the blueprint graph. */
-		TSharedPtr<FUICommandInfo> AutoCompileGraph;
 	};
 }
diff --git a/Engine/Plugins/Experimental/RenderPages/Source/RenderPagesEditor/Private/Factories/RenderPageCollectionFactory.cpp b/Engine/Plugins/Experimental/RenderPages/Source/RenderPagesEditor/Private/Factories/RenderPageCollectionFactory.cpp
--- a/Engine/Plugins/Experimental/RenderPages/Source/RenderPagesEditor/Private/Factories/RenderPageCollectionFactory.cpp
+++ b/Engine/Plugins/Experimental/RenderPages/Source/RenderPagesEditor/Private/Factories/RenderPageCollectionFactory.cpp
@@ -2,8 +2,6 @@
 
 #include "Factories/RenderPageCollectionFactory.h"
 #include "Blueprints/RenderPagesBlueprint.h"
-#include "Graph/RenderPagesGraph.h"
-#include "Graph/RenderPagesGraphSchema.h"
 #include "IContentBrowserSingleton.h"
 #include "Kismet2/BlueprintEditorUtils.h"
 #include "Kismet2/KismetEditorUtilities.h"
@@ -35,9 +33,12 @@ UObject* URenderPagesBlueprintFactory::FactoryCreateNew(UClass* InClass, UObject
 		return nullptr;
 	}
 
-	URenderPagesBlueprint* RenderPagesBlueprint = CastChecked<URenderPagesBlueprint>(FKismetEditorUtilities::CreateBlueprint(ParentClass, InParent, InName, BPTYPE_Normal, URenderPagesBlueprint::StaticClass(), URenderPagesBlueprintGeneratedClass::StaticClass(), CallingContext));
-	CreateRenderPagesGraphIfRequired(RenderPagesBlueprint);
-	return RenderPagesBlueprint;
+	if (URenderPagesBlueprint* RenderPagesBlueprint = CastChecked<URenderPagesBlueprint>(FKismetEditorUtilities::CreateBlueprint(ParentClass, InParent, InName, BPTYPE_Normal, URenderPagesBlueprint::StaticClass(), URenderPagesBlueprintGeneratedClass::StaticClass(), CallingContext)))
+	{
+		RenderPagesBlueprint->PostLoad();
+		return RenderPagesBlueprint;
+	}
+	return nullptr;
 }
 
 UObject* URenderPagesBlueprintFactory::FactoryCreateNew(UClass* InClass, UObject* InParent, FName InName, EObjectFlags Flags, UObject* Context, FFeedbackContext* Warn)
@@ -55,32 +56,8 @@ uint32 URenderPagesBlueprintFactory::GetMenuCategories() const
 	//  if wanting to show it in its own category:
 	// IAssetTools& AssetTools = FModuleManager::LoadModuleChecked<FAssetToolsModule>("AssetTools").Get();
 	// return AssetTools.RegisterAdvancedAssetCategory("Render Pages", LOCTEXT("AssetCategoryName", "Render Pages"));
-	
-	return EAssetTypeCategories::Misc;
-}
-
-void URenderPagesBlueprintFactory::CreateRenderPagesGraphIfRequired(URenderPagesBlueprint* InBlueprint)
-{
-	if (InBlueprint == nullptr)
-	{
-		return;
-	}
-
-	for (UEdGraph* EdGraph : InBlueprint->UbergraphPages)
-	{
-		if (EdGraph->IsA<URenderPagesGraph>())
-		{
-			return;
-		}
-	}
 
-	// add an initial graph for us to work in
-	const URenderPagesGraphSchema* RenderPagesGraphSchema = GetDefault<URenderPagesGraphSchema>();
-	UEdGraph* RenderPagesGraph = FBlueprintEditorUtils::CreateNewGraph(InBlueprint, RenderPagesGraphSchema->GraphName_RenderPages, URenderPagesGraph::StaticClass(), URenderPagesGraphSchema::StaticClass());
-	RenderPagesGraph->bAllowDeletion = false;
-	FBlueprintEditorUtils::AddUbergraphPage(InBlueprint, RenderPagesGraph);
-	InBlueprint->LastEditedDocuments.AddUnique(RenderPagesGraph);
-	InBlueprint->PostLoad();
+	return EAssetTypeCategories::Misc;
 }
 
 
diff --git a/Engine/Plugins/Experimental/RenderPages/Source/RenderPagesEditor/Private/Factories/RenderPageCollectionFactory.h b/Engine/Plugins/Experimental/RenderPages/Source/RenderPagesEditor/Private/Factories/RenderPageCollectionFactory.h
--- a/Engine/Plugins/Experimental/RenderPages/Source/RenderPagesEditor/Private/Factories/RenderPageCollectionFactory.h
+++ b/Engine/Plugins/Experimental/RenderPages/Source/RenderPagesEditor/Private/Factories/RenderPageCollectionFactory.h
@@ -28,12 +28,6 @@ public:
 	virtual bool ShouldShowInNewMenu() const override;
 	virtual uint32 GetMenuCategories() const override;
 
-	/**
-	 * Creates the render page graph if necessary.
-	 * @param InBlueprint The blueprint to setup up
-	 */
-	static void CreateRenderPagesGraphIfRequired(URenderPagesBlueprint* InBlueprint);
-
 	/** The parent class of the created blueprint. */
 	UPROPERTY(EditAnywhere, Category="Render Pages|Render Pages Factory", Meta = (AllowAbstract = ""))
 	TSubclassOf<URenderPageCollection> ParentClass;
diff --git a/Engine/Plugins/Experimental/RenderPages/Source/RenderPagesEditor/Private/RenderPagesEditorModule.cpp b/Engine/Plugins/Experimental/RenderPages/Source/RenderPagesEditor/Private/RenderPagesEditorModule.cpp
--- a/Engine/Plugins/Experimental/RenderPages/Source/RenderPagesEditor/Private/RenderPagesEditorModule.cpp
+++ b/Engine/Plugins/Experimental/RenderPages/Source/RenderPagesEditor/Private/RenderPagesEditorModule.cpp
@@ -111,15 +111,16 @@ TSharedPtr<UE::RenderPages::Private::SRenderPagesPropsBase> UE::RenderPages::Pri
 TSharedRef<UE::RenderPages::IRenderPageCollectionEditor> UE::RenderPages::Private::FRenderPagesEditorModule::CreateRenderPageCollectionEditor(const EToolkitMode::Type Mode, const TSharedPtr<IToolkitHost>& InitToolkitHost, URenderPagesBlueprint* InBlueprint)
 {
 	TSharedRef<FRenderPageCollectionEditor> NewRenderPagesEditor = MakeShared<FRenderPageCollectionEditor>();
-
 	NewRenderPagesEditor->InitRenderPagesEditor(Mode, InitToolkitHost, InBlueprint);
-
 	return NewRenderPagesEditor;
 }
 
 void UE::RenderPages::Private::FRenderPagesEditorModule::HandleNewBlueprintCreated(UBlueprint* InBlueprint)
 {
-	URenderPagesBlueprintFactory::CreateRenderPagesGraphIfRequired(Cast<URenderPagesBlueprint>(InBlueprint));
+	if (URenderPagesBlueprint* RenderPagesBlueprint = Cast<URenderPagesBlueprint>(InBlueprint))
+	{
+		RenderPagesBlueprint->PostLoad();
+	}
 }
 
 
diff --git a/Engine/Plugins/Experimental/RenderPages/Source/RenderPagesEditor/Private/Styles/RenderPagesEditorStyle.cpp b/Engine/Plugins/Experimental/RenderPages/Source/RenderPagesEditor/Private/Styles/RenderPagesEditorStyle.cpp
--- a/Engine/Plugins/Experimental/RenderPages/Source/RenderPagesEditor/Private/Styles/RenderPagesEditorStyle.cpp
+++ b/Engine/Plugins/Experimental/RenderPages/Source/RenderPagesEditor/Private/Styles/RenderPagesEditorStyle.cpp
@@ -73,9 +73,6 @@ TSharedRef<FSlateStyleSet> UE::RenderPages::Private::FRenderPagesEditorStyle::Cr
 		Style->SetContentRoot(FPaths::Combine(Plugin->GetBaseDir(), TEXT("Resources")));
 	}
 
-	Style->Set("RenderPages.AutoCompileGraph", new IMAGE_BRUSH("Slate/AutoCompile", Icon20x20));
-	Style->Set("RenderPages.AutoCompileGraph.Small", new IMAGE_BRUSH("Slate/AutoCompile", Icon20x20));
-
 	Style->Set("Invisible",
 		FButtonStyle()
 		.SetNormal(FSlateNoResource())
diff --git a/Engine/Plugins/Experimental/RenderPages/Source/RenderPagesEditor/Private/Toolkit/RenderPageCollectionEditor.cpp b/Engine/Plugins/Experimental/RenderPages/Source/RenderPagesEditor/Private/Toolkit/RenderPageCollectionEditor.cpp
--- a/Engine/Plugins/Experimental/RenderPages/Source/RenderPagesEditor/Private/Toolkit/RenderPageCollectionEditor.cpp
+++ b/Engine/Plugins/Experimental/RenderPages/Source/RenderPagesEditor/Private/Toolkit/RenderPageCollectionEditor.cpp
@@ -20,7 +20,6 @@
 #include "EditorModeManager.h"
 #include "Framework/Commands/GenericCommands.h"
 #include "Framework/MultiBox/MultiBoxBuilder.h"
-#include "Graph/RenderPagesGraph.h"
 #include "Kismet2/BlueprintEditorUtils.h"
 #include "Misc/MessageDialog.h"
 #include "PropertyCustomizationHelpers.h"
@@ -28,7 +27,6 @@
 #include "SGraphPanel.h"
 #include "SKismetInspector.h"
 #include "SMyBlueprint.h"
-#include "SNodePanel.h"
 #include "Stats/StatsHierarchical.h"
 #include "Styling/AppStyle.h"
 
@@ -38,96 +36,9 @@
 namespace UE::RenderPages::Private
 {
 	const FName RenderPagesEditorAppName(TEXT("RenderPagesEditorApp"));
-
-	struct FRenderPagesZoomLevelsContainer : FZoomLevelsContainer
-	{
-		struct FRenderPagesZoomLevelEntry
-		{
-		public:
-			FRenderPagesZoomLevelEntry(const float InZoomAmount, const FText& InDisplayText, const EGraphRenderingLOD::Type InLOD)
-				: DisplayText(FText::Format(NSLOCTEXT("GraphEditor", "Zoom", "Zoom {0}"), InDisplayText))
-				, ZoomAmount(InZoomAmount)
-				, LOD(InLOD)
-			{}
-
-		public:
-			FText DisplayText;
-			float ZoomAmount;
-			EGraphRenderingLOD::Type LOD;
-		};
-
-		FRenderPagesZoomLevelsContainer()
-		{
-			ZoomLevels.Reserve(22);
-			ZoomLevels.Add(FRenderPagesZoomLevelEntry(0.025f, FText::FromString(TEXT("-14")), EGraphRenderingLOD::LowestDetail));
-			ZoomLevels.Add(FRenderPagesZoomLevelEntry(0.070f, FText::FromString(TEXT("-13")), EGraphRenderingLOD::LowestDetail));
-			ZoomLevels.Add(FRenderPagesZoomLevelEntry(0.100f, FText::FromString(TEXT("-12")), EGraphRenderingLOD::LowestDetail));
-			ZoomLevels.Add(FRenderPagesZoomLevelEntry(0.125f, FText::FromString(TEXT("-11")), EGraphRenderingLOD::LowestDetail));
-			ZoomLevels.Add(FRenderPagesZoomLevelEntry(0.150f, FText::FromString(TEXT("-10")), EGraphRenderingLOD::LowestDetail));
-			ZoomLevels.Add(FRenderPagesZoomLevelEntry(0.175f, FText::FromString(TEXT("-9")), EGraphRenderingLOD::LowestDetail));
-			ZoomLevels.Add(FRenderPagesZoomLevelEntry(0.200f, FText::FromString(TEXT("-8")), EGraphRenderingLOD::LowestDetail));
-			ZoomLevels.Add(FRenderPagesZoomLevelEntry(0.225f, FText::FromString(TEXT("-7")), EGraphRenderingLOD::LowDetail));
-			ZoomLevels.Add(FRenderPagesZoomLevelEntry(0.250f, FText::FromString(TEXT("-6")), EGraphRenderingLOD::LowDetail));
-			ZoomLevels.Add(FRenderPagesZoomLevelEntry(0.375f, FText::FromString(TEXT("-5")), EGraphRenderingLOD::MediumDetail));
-			ZoomLevels.Add(FRenderPagesZoomLevelEntry(0.500f, FText::FromString(TEXT("-4")), EGraphRenderingLOD::MediumDetail));
-			ZoomLevels.Add(FRenderPagesZoomLevelEntry(0.675f, FText::FromString(TEXT("-3")), EGraphRenderingLOD::MediumDetail));
-			ZoomLevels.Add(FRenderPagesZoomLevelEntry(0.750f, FText::FromString(TEXT("-2")), EGraphRenderingLOD::DefaultDetail));
-			ZoomLevels.Add(FRenderPagesZoomLevelEntry(0.875f, FText::FromString(TEXT("-1")), EGraphRenderingLOD::DefaultDetail));
-			ZoomLevels.Add(FRenderPagesZoomLevelEntry(1.000f, FText::FromString(TEXT("1:1")), EGraphRenderingLOD::DefaultDetail));
-			ZoomLevels.Add(FRenderPagesZoomLevelEntry(1.250f, FText::FromString(TEXT("+1")), EGraphRenderingLOD::DefaultDetail));
-			ZoomLevels.Add(FRenderPagesZoomLevelEntry(1.375f, FText::FromString(TEXT("+2")), EGraphRenderingLOD::DefaultDetail));
-			ZoomLevels.Add(FRenderPagesZoomLevelEntry(1.500f, FText::FromString(TEXT("+3")), EGraphRenderingLOD::FullyZoomedIn));
-			ZoomLevels.Add(FRenderPagesZoomLevelEntry(1.675f, FText::FromString(TEXT("+4")), EGraphRenderingLOD::FullyZoomedIn));
-			ZoomLevels.Add(FRenderPagesZoomLevelEntry(1.750f, FText::FromString(TEXT("+5")), EGraphRenderingLOD::FullyZoomedIn));
-			ZoomLevels.Add(FRenderPagesZoomLevelEntry(1.875f, FText::FromString(TEXT("+6")), EGraphRenderingLOD::FullyZoomedIn));
-			ZoomLevels.Add(FRenderPagesZoomLevelEntry(2.000f, FText::FromString(TEXT("+7")), EGraphRenderingLOD::FullyZoomedIn));
-		}
-
-		virtual float GetZoomAmount(int32 InZoomLevel) const override
-		{
-			checkSlow(ZoomLevels.IsValidIndex(InZoomLevel));
-			return ZoomLevels[InZoomLevel].ZoomAmount;
-		}
-
-		virtual int32 GetNearestZoomLevel(float InZoomAmount) const override
-		{
-			for (int32 ZoomLevelIndex = 0; ZoomLevelIndex < GetNumZoomLevels(); ++ZoomLevelIndex)
-			{
-				if (InZoomAmount <= GetZoomAmount(ZoomLevelIndex))
-				{
-					return ZoomLevelIndex;
-				}
-			}
-
-			return GetDefaultZoomLevel();
-		}
-
-		virtual FText GetZoomText(int32 InZoomLevel) const override
-		{
-			checkSlow(ZoomLevels.IsValidIndex(InZoomLevel));
-			return ZoomLevels[InZoomLevel].DisplayText;
-		}
-
-		virtual int32 GetNumZoomLevels() const override
-		{
-			return ZoomLevels.Num();
-		}
-
-		virtual int32 GetDefaultZoomLevel() const override
-		{
-			return 14;
-		}
-
-		virtual EGraphRenderingLOD::Type GetLOD(int32 InZoomLevel) const override
-		{
-			checkSlow(ZoomLevels.IsValidIndex(InZoomLevel));
-			return ZoomLevels[InZoomLevel].LOD;
-		}
-
-		TArray<FRenderPagesZoomLevelEntry> ZoomLevels;
-	};
 }
 
+
 void UE::RenderPages::Private::FRenderPageCollectionEditor::InitRenderPagesEditor(const EToolkitMode::Type Mode, const TSharedPtr<IToolkitHost>& InitToolkitHost, URenderPagesBlueprint* InRenderPagesBlueprint)
 {
 	check(InRenderPagesBlueprint);
@@ -159,19 +70,6 @@ void UE::RenderPages::Private::FRenderPageCollectionEditor::InitRenderPagesEdito
 
 	CommonInitialization(RenderPagesBlueprints, false);
 
-	{
-		TArray<UEdGraph*> EdGraphs;
-		InRenderPagesBlueprint->GetAllGraphs(EdGraphs);
-
-		for (UEdGraph* Graph : EdGraphs)
-		{
-			if (URenderPagesGraph* RenderPagesGraph = Cast<URenderPagesGraph>(Graph))
-			{
-				RenderPagesGraph->Initialize(InRenderPagesBlueprint);
-			}
-		}
-	}
-
 	BindCommands();
 
 	ExtendMenu();
@@ -202,7 +100,6 @@ UE::RenderPages::Private::FRenderPageCollectionEditor::~FRenderPageCollectionEdi
 
 	if (RenderPagesBlueprint)
 	{
-		// clear editor related data from the debugged control rig instance 
 		RenderPagesBlueprint->SetObjectBeingDebugged(nullptr);
 		RenderPagesBlueprint = nullptr;
 	}
@@ -222,15 +119,6 @@ void UE::RenderPages::Private::FRenderPageCollectionEditor::CreateDefaultCommand
 	}
 }
 
-TSharedRef<SGraphEditor> UE::RenderPages::Private::FRenderPageCollectionEditor::CreateGraphEditorWidget(TSharedRef<FTabInfo> InTabInfo, UEdGraph* InGraph)
-{
-	TSharedRef<SGraphEditor> GraphEditor = FBlueprintEditor::CreateGraphEditorWidget(InTabInfo, InGraph);
-
-	GraphEditor->GetGraphPanel()->SetZoomLevelsContainer<FRenderPagesZoomLevelsContainer>();
-
-	return GraphEditor;
-}
-
 UBlueprint* UE::RenderPages::Private::FRenderPageCollectionEditor::GetBlueprintObj() const
 {
 	for (UObject* Obj : GetEditingObjects())
@@ -250,52 +138,13 @@ UBlueprint* UE::RenderPages::Private::FRenderPageCollectionEditor::GetBlueprintO
 FGraphAppearanceInfo UE::RenderPages::Private::FRenderPageCollectionEditor::GetGraphAppearance(UEdGraph* InGraph) const
 {
 	FGraphAppearanceInfo AppearanceInfo = FBlueprintEditor::GetGraphAppearance(InGraph);
-
 	if (GetBlueprintObj()->IsA(URenderPagesBlueprint::StaticClass()))
 	{
 		AppearanceInfo.CornerText = LOCTEXT("AppearanceCornerText_RenderPages", "RENDER PAGES");
 	}
-
 	return AppearanceInfo;
 }
 
-bool UE::RenderPages::Private::FRenderPageCollectionEditor::IsEditable(UEdGraph* InGraph) const
-{
-	return IsGraphInCurrentBlueprint(InGraph);
-}
-
-bool UE::RenderPages::Private::FRenderPageCollectionEditor::IsCompilingEnabled() const
-{
-	return true;
-}
-
-bool UE::RenderPages::Private::FRenderPageCollectionEditor::IsSectionVisible(NodeSectionID::Type InSectionID) const
-{
-	switch (InSectionID)
-	{
-		case NodeSectionID::GRAPH:
-		case NodeSectionID::VARIABLE:
-		case NodeSectionID::FUNCTION:
-		{
-			return true;
-		}
-		case NodeSectionID::LOCAL_VARIABLE:
-		{
-			return IsValid(Cast<URenderPagesGraph>(GetFocusedGraph()));
-		}
-		default:
-		{
-			break;
-		}
-	}
-	return false;
-}
-
-FText UE::RenderPages::Private::FRenderPageCollectionEditor::GetGraphDecorationString(UEdGraph* InGraph) const
-{
-	return FText::GetEmpty();
-}
-
 void UE::RenderPages::Private::FRenderPageCollectionEditor::OnActiveTabChanged(TSharedPtr<SDockTab> PreviouslyActive, TSharedPtr<SDockTab> NewlyActivated)
 {
 	if (!NewlyActivated.IsValid())
@@ -340,20 +189,9 @@ void UE::RenderPages::Private::FRenderPageCollectionEditor::RegisterApplicationM
 	}
 }
 
-void UE::RenderPages::Private::FRenderPageCollectionEditor::FocusInspectorOnGraphSelection(const TSet<UObject*>& NewSelection, bool bForceRefresh)
-{
-	// nothing to do here for render pages
-}
-
-bool UE::RenderPages::Private::FRenderPageCollectionEditor::CanAddNewLocalVariable() const
-{
-	return (Cast<URenderPagesGraph>(GetFocusedGraph()) != nullptr);
-}
-
 void UE::RenderPages::Private::FRenderPageCollectionEditor::Compile()
 {
 	DestroyInstance();
-
 	FBlueprintEditor::Compile();
 }
 
@@ -369,12 +207,12 @@ URenderPageCollection* UE::RenderPages::Private::FRenderPageCollectionEditor::Ge
 
 bool UE::RenderPages::Private::FRenderPageCollectionEditor::IsBatchRendering() const
 {
-	return BatchRenderJob && IsValid(BatchRenderJob);
+	return IsValid(BatchRenderJob);
 }
 
 bool UE::RenderPages::Private::FRenderPageCollectionEditor::IsPreviewRendering() const
 {
-	return PreviewRenderJob && IsValid(PreviewRenderJob);
+	return IsValid(PreviewRenderJob);
 }
 
 void UE::RenderPages::Private::FRenderPageCollectionEditor::MarkAsModified()
@@ -543,7 +381,6 @@ void UE::RenderPages::Private::FRenderPageCollectionEditor::BindCommands()
 	ToolkitCommands->MapAction(Commands.CopyPage, FExecuteAction::CreateSP(this, &FRenderPageCollectionEditor::CopyPageAction));
 	ToolkitCommands->MapAction(Commands.DeletePage, FExecuteAction::CreateSP(this, &FRenderPageCollectionEditor::DeletePageAction));
 	ToolkitCommands->MapAction(Commands.BatchRenderList, FExecuteAction::CreateSP(this, &FRenderPageCollectionEditor::BatchRenderListAction));
-	ToolkitCommands->MapAction(Commands.AutoCompileGraph, FExecuteAction::CreateSP(this, &FRenderPageCollectionEditor::AutoCompileGraphAction));
 }
 
 void UE::RenderPages::Private::FRenderPageCollectionEditor::AddPageAction()
@@ -643,8 +480,6 @@ void UE::RenderPages::Private::FRenderPageCollectionEditor::OnBatchRenderListAct
 	OnRenderPagesBatchRenderingFinished().Broadcast(FinishedRenderJob);
 }
 
-void UE::RenderPages::Private::FRenderPageCollectionEditor::AutoCompileGraphAction() {}
-
 void UE::RenderPages::Private::FRenderPageCollectionEditor::UndoAction()
 {
 	GEditor->UndoTransaction();
@@ -736,8 +571,8 @@ void UE::RenderPages::Private::FRenderPageCollectionEditor::UpdateInstance(UBlue
 			}
 		}
 
-		// Make sure the object being debugged is the preview instance
-		GetBlueprintObj()->SetObjectBeingDebugged(RenderPageCollection);
+		// Make sure the object being debugged is cleared out
+		GetBlueprintObj()->SetObjectBeingDebugged(nullptr);
 
 		// Store a reference to the preview actor.
 		RenderPageCollectionWeakPtr = RenderPageCollection;
diff --git a/Engine/Plugins/Experimental/RenderPages/Source/RenderPagesEditor/Private/Toolkit/RenderPageCollectionEditor.h b/Engine/Plugins/Experimental/RenderPages/Source/RenderPagesEditor/Private/Toolkit/RenderPageCollectionEditor.h
--- a/Engine/Plugins/Experimental/RenderPages/Source/RenderPagesEditor/Private/Toolkit/RenderPageCollectionEditor.h
+++ b/Engine/Plugins/Experimental/RenderPages/Source/RenderPagesEditor/Private/Toolkit/RenderPageCollectionEditor.h
@@ -39,23 +39,12 @@ namespace UE::RenderPages::Private
 
 		//~ Begin FBlueprintEditor Interface
 		virtual void CreateDefaultCommands() override;
-		virtual TSharedRef<SGraphEditor> CreateGraphEditorWidget(TSharedRef<FTabInfo> InTabInfo, UEdGraph* InGraph) override;
 		virtual UBlueprint* GetBlueprintObj() const override;
-		virtual TSharedPtr<FRenderPagesBlueprintEditorToolbar> GetRenderPagesToolbarBuilder() override { return RenderPagesToolbar; };
 		virtual FGraphAppearanceInfo GetGraphAppearance(UEdGraph* InGraph) const override;
-		virtual bool IsEditable(UEdGraph* InGraph) const override;
-		virtual bool IsCompilingEnabled() const override;
 		virtual bool IsInAScriptingMode() const override { return true; }
-		virtual bool IsSectionVisible(NodeSectionID::Type InSectionID) const override;
-		virtual FText GetGraphDecorationString(UEdGraph* InGraph) const override;
 		virtual void OnActiveTabChanged(TSharedPtr<SDockTab> PreviouslyActive, TSharedPtr<SDockTab> NewlyActivated) override;
 		virtual void SetupGraphEditorEvents(UEdGraph* InGraph, SGraphEditor::FGraphEditorEvents& InEvents) override;
 		virtual void RegisterApplicationModes(const TArray<UBlueprint*>& InBlueprints, bool bShouldOpenInDefaultsMode, bool bNewlyCreated = false) override;
-		virtual void FocusInspectorOnGraphSelection(const TSet<UObject*>& NewSelection, bool bForceRefresh = false) override;
-
-		virtual bool ShouldLoadBPLibrariesFromAssetRegistry() override { return false; }// Render Pages BP does not use regular BP function libraries, no need to load them.
-
-		virtual bool CanAddNewLocalVariable() const override;
 
 		virtual void Compile() override;
 		//~ End FBlueprintEditor Interface
@@ -63,6 +52,7 @@ namespace UE::RenderPages::Private
 		//~ Begin IRenderPageCollectionEditor Interface
 		virtual URenderPagesBlueprint* GetRenderPagesBlueprint() const override;
 		virtual URenderPageCollection* GetInstance() const override;
+		virtual TSharedPtr<FRenderPagesBlueprintEditorToolbar> GetRenderPagesToolbarBuilder() override { return RenderPagesToolbar; }
 		virtual bool IsBatchRendering() const override;
 		virtual URenderPagesMoviePipelineRenderJob* GetBatchRenderJob() const override { return BatchRenderJob; }
 		virtual bool IsPreviewRendering() const override;
@@ -128,9 +118,6 @@ namespace UE::RenderPages::Private
 		/** The callback for when the batch render list action finishes. */
 		void OnBatchRenderListActionFinished(URenderPagesMoviePipelineRenderJob* RenderJob, bool bSuccess);
 
-		/** Toggles auto compiling of the blueprint graph. */
-		void AutoCompileGraphAction();
-
 		/** Undo the last action. */
 		void UndoAction();
 
diff --git a/Engine/Plugins/Experimental/RenderPages/Source/RenderPagesEditor/Private/Toolkit/RenderPageCollectionEditorToolbar.cpp b/Engine/Plugins/Experimental/RenderPages/Source/RenderPagesEditor/Private/Toolkit/RenderPageCollectionEditorToolbar.cpp
--- a/Engine/Plugins/Experimental/RenderPages/Source/RenderPagesEditor/Private/Toolkit/RenderPageCollectionEditorToolbar.cpp
+++ b/Engine/Plugins/Experimental/RenderPages/Source/RenderPagesEditor/Private/Toolkit/RenderPageCollectionEditorToolbar.cpp
@@ -97,13 +97,6 @@ void UE::RenderPages::Private::FRenderPagesBlueprintEditorToolbar::AddLogicModeT
 	FToolMenuSection& Section = InMenu->AddSection("Tools");
 
 	Section.InsertPosition = FToolMenuInsert("Asset", EToolMenuInsertType::After);
-
-	Section.AddEntry(FToolMenuEntry::InitToolBarButton(
-		FRenderPagesEditorCommands::Get().AutoCompileGraph,
-		TAttribute<FText>(),
-		TAttribute<FText>(),
-		FSlateIcon(FRenderPagesEditorStyle::Get().GetStyleSetName(), "RenderPages.AutoCompileGraph")
-	));
 }
 
 void UE::RenderPages::Private::FRenderPagesBlueprintEditorToolbar::FillRenderPagesBlueprintEditorModesToolbar(FToolBarBuilder& ToolbarBuilder)
diff --git a/Engine/Plugins/Experimental/RenderPages/Source/RenderPagesEditor/Private/UI/Components/SRenderPagesRemoteControlField.cpp b/Engine/Plugins/Experimental/RenderPages/Source/RenderPagesEditor/Private/UI/Components/SRenderPagesRemoteControlField.cpp
--- a/Engine/Plugins/Experimental/RenderPages/Source/RenderPagesEditor/Private/UI/Components/SRenderPagesRemoteControlField.cpp
+++ b/Engine/Plugins/Experimental/RenderPages/Source/RenderPagesEditor/Private/UI/Components/SRenderPagesRemoteControlField.cpp
@@ -125,7 +125,6 @@ void UE::RenderPages::Private::SRenderPagesRemoteControlField::RefreshValue()
 		}
 
 		Generator->SetObjects({Objects[0]});
-
 		if (TSharedPtr<IDetailTreeNode> Node = RenderPagesWidgetUtils::FindNode(Generator->GetRootTreeNodes(), Field->FieldPathInfo.ToPathPropertyString(), RenderPagesWidgetUtils::ERenderPagesFindNodeMethod::Path))
 		{
 			TArray<TSharedRef<IDetailTreeNode>> ChildNodes;
@@ -140,11 +139,10 @@ void UE::RenderPages::Private::SRenderPagesRemoteControlField::RefreshValue()
 			//TODO:  still causes the value widgets (like the color wheel) to disconnect when this function is called,  have to somehow prevent this disconnection from happening,
 			//       skipping this line all-together causes the value widgets to not respond anymore to user input,  same happens when trying to cache and reuse the output of Node->CreateNodeWidgets()
 			ChildSlot.AttachWidget(MakeFieldWidget(ExposedFieldUtils::CreateNodeValueWidget(Node->CreateNodeWidgets())));
+			return;
 		}
-		else
-		{
-			ChildSlot.AttachWidget(MakeFieldWidget(SNullWidget::NullWidget));
-		}
+
+		ChildSlot.AttachWidget(MakeFieldWidget(SNullWidget::NullWidget));
 	}
 }
 
@@ -210,7 +208,7 @@ void UE::RenderPages::Private::SRenderPagesRemoteControlFieldChildNode::Construc
 	FNodeWidgets Widgets = InNode->CreateNodeWidgets();
 	FRenderPagesMakeNodeWidgetArgs Args;
 	Args.NameWidget = Widgets.NameWidget;
-	Args.ValueWidget = ExposedFieldUtils::CreateNodeValueWidget(InNode->CreateNodeWidgets());
+	Args.ValueWidget = ExposedFieldUtils::CreateNodeValueWidget(Widgets);
 
 	ChildSlot
 	[
diff --git a/Engine/Plugins/Experimental/RenderPages/Source/RenderPagesEditor/Private/UI/SRenderPagesCollection.cpp b/Engine/Plugins/Experimental/RenderPages/Source/RenderPagesEditor/Private/UI/SRenderPagesCollection.cpp
--- a/Engine/Plugins/Experimental/RenderPages/Source/RenderPagesEditor/Private/UI/SRenderPagesCollection.cpp
+++ b/Engine/Plugins/Experimental/RenderPages/Source/RenderPagesEditor/Private/UI/SRenderPagesCollection.cpp
@@ -13,12 +13,16 @@
 
 void UE::RenderPages::Private::SRenderPagesCollection::Tick(const FGeometry&, const double, const float)
 {
-	if(DetailsView.IsValid())
+	if (DetailsView.IsValid())
 	{
 		if (const TSharedPtr<IRenderPageCollectionEditor> BlueprintEditor = BlueprintEditorWeakPtr.Pin())
 		{
-			if (URenderPageCollection* Collection = BlueprintEditor->GetInstance(); IsValid(Collection))
+			if (URenderPageCollection* Collection = BlueprintEditor->GetInstance())
 			{
+				if (!IsValid(Collection) || BlueprintEditor->IsBatchRendering())
+				{
+					Collection = nullptr;
+				}
 				if (DetailsViewRenderPageCollectionWeakPtr != Collection)
 				{
 					DetailsViewRenderPageCollectionWeakPtr = Collection;
diff --git a/Engine/Plugins/Experimental/RenderPages/Source/RenderPagesEditor/Private/UI/SRenderPagesPage.cpp b/Engine/Plugins/Experimental/RenderPages/Source/RenderPagesEditor/Private/UI/SRenderPagesPage.cpp
--- a/Engine/Plugins/Experimental/RenderPages/Source/RenderPagesEditor/Private/UI/SRenderPagesPage.cpp
+++ b/Engine/Plugins/Experimental/RenderPages/Source/RenderPagesEditor/Private/UI/SRenderPagesPage.cpp
@@ -24,8 +24,10 @@ void UE::RenderPages::Private::SRenderPagesPage::Construct(const FArguments& InA
 	DetailsView = PropertyEditorModule.CreateDetailView(DetailsViewArgs);
 	DetailsView->SetObject(nullptr);
 
-	HandlePagesSelectionChanged();
-	InBlueprintEditor->OnRenderPagesSelectionChanged().AddSP(this, &SRenderPagesPage::HandlePagesSelectionChanged);
+	Refresh();
+	InBlueprintEditor->OnRenderPagesSelectionChanged().AddSP(this, &SRenderPagesPage::Refresh);
+	InBlueprintEditor->OnRenderPagesBatchRenderingStarted().AddSP(this, &SRenderPagesPage::OnBatchRenderingStarted);
+	InBlueprintEditor->OnRenderPagesBatchRenderingFinished().AddSP(this, &SRenderPagesPage::OnBatchRenderingFinished);
 
 	ChildSlot
 	[
@@ -46,16 +48,19 @@ void UE::RenderPages::Private::SRenderPagesPage::NotifyPostChange(const FPropert
 	}
 }
 
-void UE::RenderPages::Private::SRenderPagesPage::HandlePagesSelectionChanged()
+void UE::RenderPages::Private::SRenderPagesPage::Refresh()
 {
 	if (DetailsView.IsValid())
 	{
 		if (const TSharedPtr<IRenderPageCollectionEditor> BlueprintEditor = BlueprintEditorWeakPtr.Pin())
 		{
 			TArray<TWeakObjectPtr<UObject>> WeakSelectedPages;
-			if (const TArray<URenderPage*> SelectedPages = BlueprintEditor->GetSelectedRenderPages(); (SelectedPages.Num() == 1))
+			if (!BlueprintEditor->IsBatchRendering())
 			{
-				WeakSelectedPages.Add(SelectedPages[0]);
+				if (const TArray<URenderPage*> SelectedPages = BlueprintEditor->GetSelectedRenderPages(); (SelectedPages.Num() == 1))
+				{
+					WeakSelectedPages.Add(SelectedPages[0]);
+				}
 			}
 			DetailsView->SetObjects(WeakSelectedPages);
 		}
diff --git a/Engine/Plugins/Experimental/RenderPages/Source/RenderPagesEditor/Private/UI/SRenderPagesPage.h b/Engine/Plugins/Experimental/RenderPages/Source/RenderPagesEditor/Private/UI/SRenderPagesPage.h
--- a/Engine/Plugins/Experimental/RenderPages/Source/RenderPagesEditor/Private/UI/SRenderPagesPage.h
+++ b/Engine/Plugins/Experimental/RenderPages/Source/RenderPagesEditor/Private/UI/SRenderPagesPage.h
@@ -8,6 +8,8 @@
 
 
 class IDetailsView;
+class URenderPage;
+class URenderPagesMoviePipelineRenderJob;
 
 namespace UE::RenderPages
 {
@@ -35,8 +37,12 @@ namespace UE::RenderPages::Private
 		// End of FNotifyHook interface
 
 	private:
-		/** Callback for updating the details view. */
-		void HandlePagesSelectionChanged();
+		/** Updates the details view. */
+		void Refresh();
+
+	private:
+		void OnBatchRenderingStarted(URenderPagesMoviePipelineRenderJob* RenderJob) { Refresh(); }
+		void OnBatchRenderingFinished(URenderPagesMoviePipelineRenderJob* RenderJob) { Refresh(); }
 
 	private:
 		/** A reference to the BP Editor that owns this collection. */
diff --git a/Engine/Plugins/Experimental/RenderPages/Source/RenderPagesEditor/Private/UI/SRenderPagesPageViewer.cpp b/Engine/Plugins/Experimental/RenderPages/Source/RenderPagesEditor/Private/UI/SRenderPagesPageViewer.cpp
--- a/Engine/Plugins/Experimental/RenderPages/Source/RenderPagesEditor/Private/UI/SRenderPagesPageViewer.cpp
+++ b/Engine/Plugins/Experimental/RenderPages/Source/RenderPagesEditor/Private/UI/SRenderPagesPageViewer.cpp
@@ -24,6 +24,8 @@ void UE::RenderPages::Private::SRenderPagesPageViewer::Construct(const FArgument
 		.BorderImage(new FSlateNoResource());
 
 	Refresh();
+	InBlueprintEditor->OnRenderPagesBatchRenderingStarted().AddSP(this, &SRenderPagesPageViewer::OnBatchRenderingStarted);
+	InBlueprintEditor->OnRenderPagesBatchRenderingFinished().AddSP(this, &SRenderPagesPageViewer::OnBatchRenderingFinished);
 
 	ChildSlot
 	[
@@ -95,25 +97,30 @@ END_SLATE_FUNCTION_BUILD_OPTIMIZATION
 
 void UE::RenderPages::Private::SRenderPagesPageViewer::Refresh()
 {
-	if (!WidgetContainer.IsValid() || (ViewerMode == CachedViewerMode))
+	if (!WidgetContainer.IsValid())
 	{
 		return;
 	}
-
-	CachedViewerMode = ViewerMode;
-	WidgetContainer->ClearContent();
-
 	if (const TSharedPtr<IRenderPageCollectionEditor> BlueprintEditor = BlueprintEditorWeakPtr.Pin())
 	{
-		if (ViewerMode == ERenderPagesPageViewerMode::Live)
+		ERenderPagesPageViewerMode CurrentViewerMode = (BlueprintEditor->IsBatchRendering() ? ERenderPagesPageViewerMode::None : ViewerMode);
+		if (CurrentViewerMode == CachedViewerMode)
+		{
+			return;
+		}
+
+		CachedViewerMode = CurrentViewerMode;
+		WidgetContainer->ClearContent();
+
+		if (CurrentViewerMode == ERenderPagesPageViewerMode::Live)
 		{
 			WidgetContainer->SetContent(SNew(SRenderPagesPageViewerLive, BlueprintEditor));
 		}
-		else if (ViewerMode == ERenderPagesPageViewerMode::Preview)
+		else if (CurrentViewerMode == ERenderPagesPageViewerMode::Preview)
 		{
 			WidgetContainer->SetContent(SNew(SRenderPagesPageViewerPreview, BlueprintEditor));
 		}
-		else if (ViewerMode == ERenderPagesPageViewerMode::Rendered)
+		else if (CurrentViewerMode == ERenderPagesPageViewerMode::Rendered)
 		{
 			WidgetContainer->SetContent(SNew(SRenderPagesPageViewerRendered, BlueprintEditor));
 		}
diff --git a/Engine/Plugins/Experimental/RenderPages/Source/RenderPagesEditor/Private/UI/SRenderPagesPageViewer.h b/Engine/Plugins/Experimental/RenderPages/Source/RenderPagesEditor/Private/UI/SRenderPagesPageViewer.h
--- a/Engine/Plugins/Experimental/RenderPages/Source/RenderPagesEditor/Private/UI/SRenderPagesPageViewer.h
+++ b/Engine/Plugins/Experimental/RenderPages/Source/RenderPagesEditor/Private/UI/SRenderPagesPageViewer.h
@@ -7,6 +7,8 @@
 
 
 class SBorder;
+class URenderPage;
+class URenderPagesMoviePipelineRenderJob;
 
 namespace UE::RenderPages
 {
@@ -38,12 +40,17 @@ namespace UE::RenderPages::Private
 
 		void Construct(const FArguments& InArgs, TSharedPtr<IRenderPageCollectionEditor> InBlueprintEditor);
 
+	private:
+		/** Creates a tab button for a viewer mode. */
+		TSharedRef<SWidget> CreateViewerModeButton(const FText& ButtonText, const ERenderPagesPageViewerMode ButtonViewerMode);
+
+	public:
 		/** Refreshes the content of this widget. */
 		void Refresh();
 
 	private:
-		/** Creates a tab button for a viewer mode. */
-		TSharedRef<SWidget> CreateViewerModeButton(const FText& ButtonText, const ERenderPagesPageViewerMode ButtonViewerMode);
+		void OnBatchRenderingStarted(URenderPagesMoviePipelineRenderJob* RenderJob) { Refresh(); }
+		void OnBatchRenderingFinished(URenderPagesMoviePipelineRenderJob* RenderJob) { Refresh(); }
 
 	private:
 		/** A reference to the BP Editor that owns this collection. */
diff --git a/Engine/Plugins/Experimental/RenderPages/Source/RenderPagesEditor/Private/UI/SRenderPagesPageViewerLive.cpp b/Engine/Plugins/Experimental/RenderPages/Source/RenderPagesEditor/Private/UI/SRenderPagesPageViewerLive.cpp
--- a/Engine/Plugins/Experimental/RenderPages/Source/RenderPagesEditor/Private/UI/SRenderPagesPageViewerLive.cpp
+++ b/Engine/Plugins/Experimental/RenderPages/Source/RenderPagesEditor/Private/UI/SRenderPagesPageViewerLive.cpp
@@ -11,7 +11,6 @@
 #include "RenderPage/RenderPageCollection.h"
 #include "SlateOptMacros.h"
 #include "UI/Components/SRenderPagesPageViewerFrameSlider.h"
-#include "Utils/RenderPageLevelSequencePlayer.h"
 #include "Widgets/Layout/SScaleBox.h"
 
 #define LOCTEXT_NAMESPACE "SRenderPagesPageViewerLive"
@@ -139,7 +138,7 @@ ULevelSequencePlayer* UE::RenderPages::Private::SRenderPagesEditorViewport::GetS
 		FLevelSequenceCameraSettings CameraSettings;
 
 		ALevelSequenceActor* PlayerActor = nullptr;
-		if (ULevelSequencePlayer* Player = URenderPageLevelSequencePlayer::CreateLevelSequencePlayer(World, LevelSequence, PlaybackSettings, PlayerActor); IsValid(Player))
+		if (ULevelSequencePlayer* Player = ULevelSequencePlayer::CreateLevelSequencePlayer(World, LevelSequence, PlaybackSettings, PlayerActor); IsValid(Player))
 		{
 			if (IsValid(PlayerActor))
 			{
diff --git a/Engine/Plugins/Experimental/RenderPages/Source/RenderPagesEditor/Private/UI/SRenderPagesProps.cpp b/Engine/Plugins/Experimental/RenderPages/Source/RenderPagesEditor/Private/UI/SRenderPagesProps.cpp
--- a/Engine/Plugins/Experimental/RenderPages/Source/RenderPagesEditor/Private/UI/SRenderPagesProps.cpp
+++ b/Engine/Plugins/Experimental/RenderPages/Source/RenderPagesEditor/Private/UI/SRenderPagesProps.cpp
@@ -37,6 +37,8 @@ void UE::RenderPages::Private::SRenderPagesProps::Construct(const FArguments& In
 
 	Refresh();
 	InBlueprintEditor->OnRenderPagesChanged().AddSP(this, &SRenderPagesProps::Refresh);
+	InBlueprintEditor->OnRenderPagesBatchRenderingStarted().AddSP(this, &SRenderPagesProps::OnBatchRenderingStarted);
+	InBlueprintEditor->OnRenderPagesBatchRenderingFinished().AddSP(this, &SRenderPagesProps::OnBatchRenderingFinished);
 
 	ChildSlot
 	[
@@ -56,14 +58,17 @@ void UE::RenderPages::Private::SRenderPagesProps::Refresh()
 
 	if (const TSharedPtr<IRenderPageCollectionEditor> BlueprintEditor = BlueprintEditorWeakPtr.Pin())
 	{
-		if (URenderPageCollection* Collection = BlueprintEditor->GetInstance(); IsValid(Collection))
+		if (!BlueprintEditor->IsBatchRendering())
 		{
-			WidgetPropsSourceWeakPtr = Collection->GetPropsSource();
-			if (URenderPagePropsSourceBase* WidgetPropsSource = WidgetPropsSourceWeakPtr.Get(); IsValid(WidgetPropsSource))
+			if (URenderPageCollection* Collection = BlueprintEditor->GetInstance(); IsValid(Collection))
 			{
-				if (const TSharedPtr<SRenderPagesPropsBase> Widget = IRenderPagesEditorModule::Get().CreatePropsSourceWidget(WidgetPropsSource, BlueprintEditor))
+				WidgetPropsSourceWeakPtr = Collection->GetPropsSource();
+				if (URenderPagePropsSourceBase* WidgetPropsSource = WidgetPropsSourceWeakPtr.Get(); IsValid(WidgetPropsSource))
 				{
-					WidgetContainer->SetContent(Widget.ToSharedRef());
+					if (const TSharedPtr<SRenderPagesPropsBase> Widget = IRenderPagesEditorModule::Get().CreatePropsSourceWidget(WidgetPropsSource, BlueprintEditor))
+					{
+						WidgetContainer->SetContent(Widget.ToSharedRef());
+					}
 				}
 			}
 		}
diff --git a/Engine/Plugins/Experimental/RenderPages/Source/RenderPagesEditor/Private/UI/SRenderPagesProps.h b/Engine/Plugins/Experimental/RenderPages/Source/RenderPagesEditor/Private/UI/SRenderPagesProps.h
--- a/Engine/Plugins/Experimental/RenderPages/Source/RenderPagesEditor/Private/UI/SRenderPagesProps.h
+++ b/Engine/Plugins/Experimental/RenderPages/Source/RenderPagesEditor/Private/UI/SRenderPagesProps.h
@@ -7,7 +7,9 @@
 
 
 class SBorder;
+class URenderPage;
 class URenderPagePropsSourceBase;
+class URenderPagesMoviePipelineRenderJob;
 
 namespace UE::RenderPages
 {
@@ -33,6 +35,10 @@ namespace UE::RenderPages::Private
 		/** Refreshes the content of this widget. */
 		void Refresh();
 
+	private:
+		void OnBatchRenderingStarted(URenderPagesMoviePipelineRenderJob* RenderJob) { Refresh(); }
+		void OnBatchRenderingFinished(URenderPagesMoviePipelineRenderJob* RenderJob) { Refresh(); }
+
 	private:
 		/** A reference to the BP Editor that owns this collection. */
 		TWeakPtr<IRenderPageCollectionEditor> BlueprintEditorWeakPtr;
diff --git a/Engine/Plugins/Experimental/RenderPages/Source/RenderPagesEditor/Public/IRenderPageCollectionEditor.h b/Engine/Plugins/Experimental/RenderPages/Source/RenderPagesEditor/Public/IRenderPageCollectionEditor.h
--- a/Engine/Plugins/Experimental/RenderPages/Source/RenderPagesEditor/Public/IRenderPageCollectionEditor.h
+++ b/Engine/Plugins/Experimental/RenderPages/Source/RenderPagesEditor/Public/IRenderPageCollectionEditor.h
@@ -79,7 +79,7 @@ namespace UE::RenderPages
 		DECLARE_MULTICAST_DELEGATE_OneParam(FOnRenderPagesBatchRenderingFinished, URenderPagesMoviePipelineRenderJob*);
 		virtual FOnRenderPagesBatchRenderingFinished& OnRenderPagesBatchRenderingFinished() { return OnRenderPagesBatchRenderingFinishedDelegate; }
 
-	protected:
+	private:
 		/** The delegate for when render pages in the collection changed. */
 		FOnRenderPagesChanged OnRenderPagesChangedDelegate;
 
