- Moved render page manager functions over to the render page collection itself.
- Made most of the page and page collection functions usable in blueprint.
- Made the pre and post render events be executed in the actual PIE game world (so blueprint nodes such as "Get All Actors" will work correctly).
- Made the pre and post render events work in the live viewer as well (the functions are now called during and after a live viewer frame, the world it will execute in is the editor world).
- Made it automatically make new variables in the render page graph "Transient", and also make any "Instance Editable" variables automatically non-"Transient" again (this prevents common mistakes, while still giving the user full control to override this if desired).
- Made the page properties be automatically returned to the user-set values after rendering (in case the blueprint code of the user overrides any values that are referenced by the page properties).

diff --git a/Engine/Plugins/Experimental/RenderPages/Source/RenderPages/Private/RenderPage/RenderPageCollection.cpp b/Engine/Plugins/Experimental/RenderPages/Source/RenderPages/Private/RenderPage/RenderPageCollection.cpp
--- a/Engine/Plugins/Experimental/RenderPages/Source/RenderPages/Private/RenderPage/RenderPageCollection.cpp
+++ b/Engine/Plugins/Experimental/RenderPages/Source/RenderPages/Private/RenderPage/RenderPageCollection.cpp
@@ -8,6 +8,7 @@
 #include "MoviePipelineMasterConfig.h"
 #include "MoviePipelineOutputSetting.h"
 #include "MovieScene.h"
+#include "RenderPagesModule.h"
 #include "UObject/ObjectSaveContext.h"
 
 
@@ -80,7 +81,7 @@ TOptional<int32> URenderPage::GetSequenceEndFrame() const
 bool URenderPage::SetSequenceStartFrame(const int32 NewCustomStartFrame)
 {
 	int32 StartFrame = GetStartFrame().Get(0);
-	SetIsCustomStartFrame(true);
+	SetIsUsingCustomStartFrame(true);
 	SetCustomStartFrame(StartFrame);
 
 	TOptional<int32> SequenceStartFrame = GetSequenceStartFrame();
@@ -110,7 +111,7 @@ bool URenderPage::SetSequenceStartFrame(const int32 NewCustomStartFrame)
 bool URenderPage::SetSequenceEndFrame(const int32 NewCustomStartFrame)
 {
 	int32 EndFrame = GetEndFrame().Get(0);
-	SetIsCustomEndFrame(true);
+	SetIsUsingCustomEndFrame(true);
 	SetCustomEndFrame(EndFrame);
 
 	TOptional<int32> SequenceEndFrame = GetSequenceEndFrame();
@@ -339,7 +340,7 @@ FString URenderPage::PurgePageIdOrGenerateUniqueId(URenderPageCollection* PageCo
 	FString Result = PurgePageIdOrReturnEmptyString(NewPageId);
 	if (Result.IsEmpty())
 	{
-		return UE::RenderPages::IRenderPagesModule::Get().GetManager().CreateUniquePageId(PageCollection);
+		return PageCollection->GenerateNextPageId();
 	}
 	return Result;
 }
@@ -449,6 +450,8 @@ URenderPageCollection::URenderPageCollection()
 	: Id(FGuid::NewGuid())
 	, PropsSourceType(ERenderPagePropsSourceType::Local)
 	, PropsSourceOrigin_RemoteControl(nullptr)
+	, bExecutingPreRender(false)
+	, bExecutingPostRender(false)
 	, CachedPropsSource(nullptr)
 	, CachedPropsSourceType(ERenderPagePropsSourceType::Local)
 {
@@ -461,24 +464,34 @@ URenderPageCollection::URenderPageCollection()
 
 UWorld* URenderPageCollection::GetWorld() const
 {
-	if (UWorld* CachedWorld = CachedWorldWeakPtr.Get(); IsValid(CachedWorld))
-	{
-		return CachedWorld;
-	}
-
 	if (HasAllFlags(RF_ClassDefaultObject))
 	{
 		// If we are a CDO, we must return nullptr instead of calling Outer->GetWorld() to fool UObject::ImplementsGetWorld.
 		return nullptr;
 	}
 
-	// Could be a GameInstance, could be World, could also be a WidgetTree, so we're just going to follow
-	// the outer chain to find the world we're in.
-	UObject* Outer = GetOuter();
+	for (const FWorldContext& Context : GEngine->GetWorldContexts())
+	{
+		if (Context.WorldType == EWorldType::PIE)
+		{
+			return Context.World();
+		}
+	}
+
+	if (IsValid(GWorld))
+	{
+		return GWorld;
+	}
+
+	if (UWorld* CachedWorld = CachedWorldWeakPtr.Get(); IsValid(CachedWorld))
+	{
+		return CachedWorld;
+	}
+	UObject* Outer = GetOuter();// Could be a GameInstance, could be World, could also be a WidgetTree, so we're just going to follow the outer chain to find the world we're in.
 	while (Outer)
 	{
 		UWorld* World = Outer->GetWorld();
-		if (World)
+		if (IsValid(World))
 		{
 			CachedWorldWeakPtr = World;
 			return World;
@@ -506,12 +519,17 @@ void URenderPageCollection::PostLoad()
 
 void URenderPageCollection::PreRender(URenderPage* Page)
 {
+	bExecutingPreRender = true;
 	ReceivePreRender(Page);
+	bExecutingPreRender = false;
 }
 
 void URenderPageCollection::PostRender(URenderPage* Page)
 {
+	bExecutingPostRender = true;
 	ReceivePostRender(Page);
+	UE::RenderPages::Private::FRenderPagesModule::Get().GetManager().ApplyPagePropValues(this, Page);
+	bExecutingPostRender = false;
 }
 
 void URenderPageCollection::CopyValuesToOrFromCDO(const bool bToCDO)
@@ -529,7 +547,7 @@ void URenderPageCollection::CopyValuesToOrFromCDO(const bool bToCDO)
 
 	for (FProperty* Property = GetClass()->PropertyLink; Property; Property = Property->PropertyLinkNext)
 	{
-		if (!Property->HasAnyPropertyFlags(CPF_Transient | CPF_DuplicateTransient))
+		if (!Property->HasAnyPropertyFlags(CPF_Transient | CPF_DuplicateTransient))// if not [Transient]
 		{
 			void* Data = Property->ContainerPtrToValuePtr<void>(this);
 			void* DataCDO = Property->ContainerPtrToValuePtr<void>(CDO);
@@ -706,3 +724,102 @@ void URenderPageCollection::InsertRenderPageAfter(URenderPage* RenderPage, URend
 		}
 	}
 }
+
+FString URenderPageCollection::GenerateNextPageId()
+{
+	int32 Max = 0;
+	for (URenderPage* Page : RenderPages)
+	{
+		if (!IsValid(Page))
+		{
+			continue;
+		}
+		int32 Value = FCString::Atoi(*Page->GetPageId());
+		if (Value > Max)
+		{
+			Max = Value;
+		}
+	}
+	FString Result = FString::FromInt(Max + 1);
+	while (Result.Len() < UE::RenderPages::FRenderPageManager::GeneratedIdCharacterLength)
+	{
+		Result = TEXT("0") + Result;
+	}
+	return Result;
+}
+
+bool URenderPageCollection::DoesPageIdExist(const FString& PageId)
+{
+	const FString PageIdToLower = PageId.ToLower();
+	for (URenderPage* Page : RenderPages)
+	{
+		if (!IsValid(Page))
+		{
+			continue;
+		}
+		if (PageIdToLower == Page->GetPageId().ToLower())
+		{
+			return true;
+		}
+	}
+	return false;
+}
+
+URenderPage* URenderPageCollection::CreateAndAddNewRenderPage()
+{
+	URenderPage* Page = NewObject<URenderPage>(this);
+	Page->SetPageId(GenerateNextPageId());
+	Page->SetPageName(TEXT("New"));
+	Page->SetOutputDirectory(FPaths::ProjectDir() / TEXT("Saved/MovieRenders/"));
+
+	if (URenderPagePropsSourceRemoteControl* PropsSource = GetPropsSource<URenderPagePropsSourceRemoteControl>())
+	{
+		TArray<uint8> BinaryArray;
+		for (URenderPagePropRemoteControl* Field : PropsSource->GetProps()->GetAllCasted())
+		{
+			if (Field->GetValue(BinaryArray))
+			{
+				Page->SetRemoteControlValue(Field->GetRemoteControlEntity(), BinaryArray);
+			}
+		}
+	}
+
+	AddRenderPage(Page);
+	return Page;
+}
+
+URenderPage* URenderPageCollection::DuplicateAndAddRenderPage(URenderPage* Page)
+{
+	if (!IsValid(Page))
+	{
+		return nullptr;
+	}
+
+	if (URenderPage* DuplicateRenderPage = DuplicateObject(Page, this); IsValid(DuplicateRenderPage))
+	{
+		DuplicateRenderPage->GenerateNewId();
+		DuplicateRenderPage->SetPageId(GenerateNextPageId());
+		InsertRenderPageAfter(DuplicateRenderPage, Page);
+		return DuplicateRenderPage;
+	}
+	return nullptr;
+}
+
+bool URenderPageCollection::ReorderRenderPage(URenderPage* Page, URenderPage* DroppedOnPage, const bool bAfter)
+{
+	if (!IsValid(Page) || !IsValid(DroppedOnPage) || !HasRenderPage(Page) || !HasRenderPage(DroppedOnPage))
+	{
+		return false;
+	}
+
+	RemoveRenderPage(Page);
+	if (bAfter)
+	{
+		InsertRenderPageAfter(Page, DroppedOnPage);
+	}
+	else
+	{
+		InsertRenderPageBefore(Page, DroppedOnPage);
+	}
+	return true;
+}
diff --git a/Engine/Plugins/Experimental/RenderPages/Source/RenderPages/Private/RenderPage/RenderPageManager.cpp b/Engine/Plugins/Experimental/RenderPages/Source/RenderPages/Private/RenderPage/RenderPageManager.cpp
--- a/Engine/Plugins/Experimental/RenderPages/Source/RenderPages/Private/RenderPage/RenderPageManager.cpp
+++ b/Engine/Plugins/Experimental/RenderPages/Source/RenderPages/Private/RenderPage/RenderPageManager.cpp
@@ -7,129 +7,6 @@
 #include "MoviePipelinePIEExecutor.h"
 
 
-URenderPage* UE::RenderPages::FRenderPageManager::AddNewPage(URenderPageCollection* PageCollection)
-{
-	if (!IsValid(PageCollection))
-	{
-		return nullptr;
-	}
-
-	URenderPage* Page = NewObject<URenderPage>(PageCollection);
-	Page->SetPageId(CreateUniquePageId(PageCollection));
-	Page->SetPageName(TEXT("New"));
-	Page->SetOutputDirectory(FPaths::ProjectDir() / TEXT("Saved/MovieRenders/"));
-
-	if (URenderPagePropsSourceRemoteControl* PropsSource = PageCollection->GetPropsSource<URenderPagePropsSourceRemoteControl>())
-	{
-		TArray<uint8> BinaryArray;
-		for (URenderPagePropRemoteControl* Field : PropsSource->GetProps()->GetAllCasted())
-		{
-			if (Field->GetValue(BinaryArray))
-			{
-				Page->SetRemoteControlValue(Field->GetRemoteControlEntity(), BinaryArray);
-			}
-		}
-	}
-
-	PageCollection->Modify();
-	PageCollection->AddRenderPage(Page);
-	return Page;
-}
-
-URenderPage* UE::RenderPages::FRenderPageManager::CopyPage(URenderPageCollection* PageCollection, URenderPage* Page)
-{
-	if (!IsValid(PageCollection) || !IsValid(Page))
-	{
-		return nullptr;
-	}
-
-	if (URenderPage* DuplicateRenderPage = DuplicateObject(Page, PageCollection); IsValid(DuplicateRenderPage))
-	{
-		DuplicateRenderPage->GenerateNewId();
-		DuplicateRenderPage->SetPageId(CreateUniquePageId(PageCollection));
-		DuplicateRenderPage->Modify();
-		PageCollection->Modify();
-		PageCollection->InsertRenderPageAfter(DuplicateRenderPage, Page);
-		return DuplicateRenderPage;
-	}
-	return nullptr;
-}
-
-bool UE::RenderPages::FRenderPageManager::DoesPageIdExist(URenderPageCollection* PageCollection, const FString& PageId)
-{
-	const FString PageIdToLower = PageId.ToLower();
-	for (URenderPage* Page : PageCollection->GetRenderPagesRef())
-	{
-		if (!IsValid(Page))
-		{
-			continue;
-		}
-		if (PageIdToLower == Page->GetPageId().ToLower())
-		{
-			return true;
-		}
-	}
-	return false;
-}
-
-void UE::RenderPages::FRenderPageManager::DeletePage(URenderPageCollection* PageCollection, URenderPage* Page)
-{
-	if (!IsValid(PageCollection) || !IsValid(Page))
-	{
-		return;
-	}
-
-	PageCollection->Modify();
-	PageCollection->RemoveRenderPage(Page);
-}
-
-bool UE::RenderPages::FRenderPageManager::DragDropPage(URenderPageCollection* PageCollection, URenderPage* Page, URenderPage* DroppedOnPage, const bool bAfter)
-{
-	if (!IsValid(PageCollection) || !IsValid(Page) || !IsValid(DroppedOnPage))
-	{
-		return false;
-	}
-	if (!PageCollection->HasRenderPage(Page) || !PageCollection->HasRenderPage(DroppedOnPage))
-	{
-		return false;
-	}
-
-	PageCollection->RemoveRenderPage(Page);
-	if (bAfter)
-	{
-		PageCollection->InsertRenderPageAfter(Page, DroppedOnPage);
-	}
-	else
-	{
-		PageCollection->InsertRenderPageBefore(Page, DroppedOnPage);
-	}
-	return true;
-}
-
-FString UE::RenderPages::FRenderPageManager::CreateUniquePageId(URenderPageCollection* PageCollection)
-{
-	int32 Max = 0;
-	for (URenderPage* Page : PageCollection->GetRenderPagesRef())
-	{
-		if (!IsValid(Page))
-		{
-			continue;
-		}
-		int32 Value = FCString::Atoi(*Page->GetPageId());
-		if (Value > Max)
-		{
-			Max = Value;
-		}
-	}
-	FString Result = FString::FromInt(Max + 1);
-	while (Result.Len() < GeneratedIdCharacterLength)
-	{
-		Result = TEXT("0") + Result;
-	}
-	return Result;
-}
-
-
 URenderPagesMoviePipelineRenderJob* UE::RenderPages::FRenderPageManager::CreateBatchRenderJob(URenderPageCollection* PageCollection)
 {
 	FRenderPagesMoviePipelineRenderJobCreateArgs JobArgs;
@@ -167,10 +44,10 @@ URenderPagesMoviePipelineRenderJob* UE::RenderPages::FRenderPageManager::RenderP
 	{
 		constexpr int32 RenderFramesCount = 1;// can be more than 1 to prevent rendering issues, will always take the last frame that's rendered
 
-		PageCopy->SetIsCustomStartFrame(true);
+		PageCopy->SetIsUsingCustomStartFrame(true);
 		PageCopy->SetCustomStartFrame(Args.Frame.Get(0));
 
-		PageCopy->SetIsCustomEndFrame(true);
+		PageCopy->SetIsUsingCustomEndFrame(true);
 		PageCopy->SetCustomEndFrame(Args.Frame.Get(0));
 
 		if (!PageCopy->SetSequenceEndFrame(PageCopy->GetSequenceStartFrame().Get(0) + 1) ||
@@ -181,7 +58,7 @@ URenderPagesMoviePipelineRenderJob* UE::RenderPages::FRenderPageManager::RenderP
 		}
 	}
 
-	PageCopy->SetIsCustomResolution(true);
+	PageCopy->SetIsUsingCustomResolution(true);
 	PageCopy->SetCustomResolution(Args.Resolution);
 
 	PageCopy->SetOutputDirectory(TmpRenderedFramesPath / (Args.Frame.IsSet() ? TEXT("PreviewFrame") : TEXT("PreviewFrames")));
@@ -206,7 +83,7 @@ URenderPagesMoviePipelineRenderJob* UE::RenderPages::FRenderPageManager::RenderP
 
 	const FGuid PageId = PageCopy->GetId();
 	const TOptional<int32> StartFrameOfRender = (Args.Frame.IsSet() ? TOptional<int32>() : PageCopy->GetStartFrame());
-	NewRenderJob->OnExecuteFinished().AddLambda([this,Callback,PageId,StartFrameOfRender](URenderPagesMoviePipelineRenderJob* RenderJob, const bool bSuccess)
+	NewRenderJob->OnExecuteFinished().AddLambda([this, Callback, PageId, StartFrameOfRender](URenderPagesMoviePipelineRenderJob* RenderJob, const bool bSuccess)
 	{
 		if (StartFrameOfRender.IsSet())
 		{
diff --git a/Engine/Plugins/Experimental/RenderPages/Source/RenderPages/Private/RenderPage/RenderPageMoviePipelineJob.cpp b/Engine/Plugins/Experimental/RenderPages/Source/RenderPages/Private/RenderPage/RenderPageMoviePipelineJob.cpp
--- a/Engine/Plugins/Experimental/RenderPages/Source/RenderPages/Private/RenderPage/RenderPageMoviePipelineJob.cpp
+++ b/Engine/Plugins/Experimental/RenderPages/Source/RenderPages/Private/RenderPage/RenderPageMoviePipelineJob.cpp
@@ -36,6 +36,8 @@ URenderPagesMoviePipelineRenderJobEntry* URenderPagesMoviePipelineRenderJobEntry
 	}
 
 	URenderPagesMoviePipelineRenderJobEntry* RenderJobEntry = NewObject<URenderPagesMoviePipelineRenderJobEntry>(Job);
+	RenderJobEntry->Page = Page;
+	RenderJobEntry->PageCollection = Args.PageCollection;
 	RenderJobEntry->RenderQueue = NewObject<UMoviePipelineQueue>(RenderJobEntry);
 	RenderJobEntry->Executor = NewObject<UMoviePipelineExecutorBase>(RenderJobEntry, PipelineExecutor);
 	RenderJobEntry->ExecutorJob = nullptr;
@@ -162,7 +164,7 @@ URenderPagesMoviePipelineRenderJobEntry* URenderPagesMoviePipelineRenderJobEntry
 			Setting->bUseCustomFrameRate = false;
 		}
 
-		if (Page->GetIsCustomResolution())
+		if (Page->GetIsUsingCustomResolution())
 		{
 			Setting->OutputResolution = Page->GetCustomResolution();
 		}
@@ -233,6 +235,11 @@ TSharedFuture<void> URenderPagesMoviePipelineRenderJobEntry::Execute()
 	{
 		LevelSequenceEditorModule->OnComputePlaybackContext().AddUObject(this, &URenderPagesMoviePipelineRenderJobEntry::ComputePlaybackContext);
 	}
+	if (UMoviePipelinePIEExecutor* ActiveExecutorPIE = Cast<UMoviePipelinePIEExecutor>(Executor))
+	{
+		ActiveExecutorPIE->OnIndividualJobStarted().AddUObject(this, &URenderPagesMoviePipelineRenderJobEntry::ExecutePageStarted);
+		ActiveExecutorPIE->OnIndividualJobWorkFinished().AddUObject(this, &URenderPagesMoviePipelineRenderJobEntry::ExecutePageFinished);
+	}
 	Status = TEXT("Rendering...");
 	Executor->OnExecutorFinished().AddUObject(this, &URenderPagesMoviePipelineRenderJobEntry::ExecuteFinished);
 	Executor->Execute(RenderQueue);
@@ -277,12 +284,27 @@ void URenderPagesMoviePipelineRenderJobEntry::ComputePlaybackContext(bool& bOutA
 	bOutAllowBinding = false;
 }
 
+void URenderPagesMoviePipelineRenderJobEntry::ExecutePageStarted(UMoviePipelineExecutorJob* JobToStart)
+{
+	PageCollection->PreRender(Page);
+}
+
+void URenderPagesMoviePipelineRenderJobEntry::ExecutePageFinished(FMoviePipelineOutputData PipelineOutputData)
+{
+	PageCollection->PostRender(Page);
+}
+
 void URenderPagesMoviePipelineRenderJobEntry::ExecuteFinished(UMoviePipelineExecutorBase* PipelineExecutor, const bool bSuccess)
 {
 	if (ILevelSequenceEditorModule* LevelSequenceEditorModule = FModuleManager::GetModulePtr<ILevelSequenceEditorModule>("LevelSequenceEditorModule"))
 	{
 		LevelSequenceEditorModule->OnComputePlaybackContext().RemoveAll(this);
 	}
+	if (UMoviePipelinePIEExecutor* PipelineExecutorPIE = Cast<UMoviePipelinePIEExecutor>(PipelineExecutor))
+	{
+		PipelineExecutorPIE->OnIndividualJobStarted().RemoveAll(this);
+		PipelineExecutorPIE->OnIndividualJobWorkFinished().RemoveAll(this);
+	}
 	bCanceled = (bCanceled || !bSuccess);
 	Status = (bCanceled ? TEXT("Canceled") : TEXT("Done"));
 	if (Promise.IsValid())
@@ -311,7 +333,6 @@ URenderPagesMoviePipelineRenderJob* URenderPagesMoviePipelineRenderJob::Create(c
 	RenderJob->Queue = MakeShareable(new UE::RenderPages::Private::FRenderPageQueue);
 	RenderJob->PageCollection = Args.PageCollection;
 	RenderJob->bCanceled = false;
-	RenderJob->bRanPreRender = false;
 
 
 	RenderJob->Queue->Add(UE::RenderPages::Private::FRenderPageQueueAction::CreateLambda([RenderJob]()
@@ -326,23 +347,12 @@ URenderPagesMoviePipelineRenderJob* URenderPagesMoviePipelineRenderJob::Create(c
 		{
 			RenderJob->Entries.Add(Page, Entry);
 
-			RenderJob->Queue->Add(UE::RenderPages::Private::FRenderPageQueueActionReturningDelay::CreateLambda([RenderJob, Page]() -> UE::RenderPages::Private::FRenderPageQueueDelay
-			{
-				if (!RenderJob->IsCanceled())
-				{
-					RenderJob->bRanPreRender = true;
-					RenderJob->PageCollection->PreRender(Page);
-					return UE::RenderPages::Private::FRenderPageQueueDelay::Frames(1);
-				}
-				return nullptr;
-			}));
-
 			RenderJob->Queue->Add(UE::RenderPages::Private::FRenderPageQueueActionReturningDelay::CreateLambda([RenderJob, Page]() -> UE::RenderPages::Private::FRenderPageQueueDelay
 			{
 				if (!RenderJob->IsCanceled())
 				{
 					RenderJob->PreviousPageProps = UE::RenderPages::IRenderPagesModule::Get().GetManager().ApplyPagePropValues(RenderJob->PageCollection, Page.Get());
-					return UE::RenderPages::Private::FRenderPageQueueDelay::Frames(1 + Page->GetWaitFramesBeforeRendering());
+					return UE::RenderPages::Private::FRenderPageQueueDelay::Frames(2 + Page->GetWaitFramesBeforeRendering());
 				}
 				return nullptr;
 			}));
@@ -358,18 +368,7 @@ URenderPagesMoviePipelineRenderJob* URenderPagesMoviePipelineRenderJob::Create(c
 				{
 					UE::RenderPages::IRenderPagesModule::Get().GetManager().RestorePagePropValues(RenderJob->PreviousPageProps);
 					RenderJob->PreviousPageProps = FRenderPageManagerPreviousPagePropValues();
-					return UE::RenderPages::Private::FRenderPageQueueDelay::Frames(1);
-				}
-				return nullptr;
-			}));
-
-			RenderJob->Queue->Add(UE::RenderPages::Private::FRenderPageQueueActionReturningDelay::CreateLambda([RenderJob, Page]() -> UE::RenderPages::Private::FRenderPageQueueDelay
-			{
-				if (RenderJob->bRanPreRender)
-				{
-					RenderJob->bRanPreRender = false;
-					RenderJob->PageCollection->PostRender(Page);
-					return UE::RenderPages::Private::FRenderPageQueueDelay::Frames(1);
+					return UE::RenderPages::Private::FRenderPageQueueDelay::Frames(2);
 				}
 				return nullptr;
 			}));
diff --git a/Engine/Plugins/Experimental/RenderPages/Source/RenderPages/Public/RenderPage/RenderPageCollection.h b/Engine/Plugins/Experimental/RenderPages/Source/RenderPages/Public/RenderPage/RenderPageCollection.h
--- a/Engine/Plugins/Experimental/RenderPages/Source/RenderPages/Public/RenderPage/RenderPageCollection.h
+++ b/Engine/Plugins/Experimental/RenderPages/Source/RenderPages/Public/RenderPage/RenderPageCollection.h
@@ -79,124 +79,124 @@ public:
 	TOptional<double> GetDurationInSeconds() const;
 
 	/** Gets the aspect ratio that this page will be rendered in. */
-	UFUNCTION(BlueprintPure, Category="Render Pages | Page")
+	UFUNCTION(BlueprintPure, Category="Render Pages|Page")
 	double GetOutputAspectRatio() const;
 
 	/** Checks whether the page contains data that matches the search terms. */
-	UFUNCTION(BlueprintPure, Category="Render Pages | Page")
+	UFUNCTION(BlueprintPure, Category="Render Pages|Page")
 	bool MatchesSearchTerm(const FString& SearchTerm) const;
 
 public:
-	UFUNCTION(BlueprintPure, Category="Render Pages | Page")
+	UFUNCTION(BlueprintPure, Category="Render Pages|Page")
 	FGuid GetId() const { return Id; }
 
-	UFUNCTION(BlueprintCallable, Category="Render Pages | Page")
+	UFUNCTION(BlueprintCallable, Category="Render Pages|Page")
 	void GenerateNewId() { Id = FGuid::NewGuid(); }
 
 
-	UFUNCTION(BlueprintPure, Category="Render Pages | Page")
+	UFUNCTION(BlueprintPure, Category="Render Pages|Page")
 	int32 GetWaitFramesBeforeRendering() const { return WaitFramesBeforeRendering; }
 
-	UFUNCTION(BlueprintCallable, Category="Render Pages | Page")
+	UFUNCTION(BlueprintCallable, Category="Render Pages|Page")
 	void SetWaitFramesBeforeRendering(const int32 NewWaitFramesBeforeRendering) { WaitFramesBeforeRendering = FMath::Max<int32>(0, NewWaitFramesBeforeRendering); }
 
 
-	UFUNCTION(BlueprintPure, Category="Render Pages | Page")
+	UFUNCTION(BlueprintPure, Category="Render Pages|Page")
 	ULevelSequence* GetSequence() const { return Sequence; }
 
-	UFUNCTION(BlueprintCallable, Category="Render Pages | Page")
+	UFUNCTION(BlueprintCallable, Category="Render Pages|Page")
 	void SetSequence(ULevelSequence* NewSequence) { Sequence = NewSequence; }
 
 
-	UFUNCTION(BlueprintPure, Category="Render Pages | Page")
-	bool GetIsCustomStartFrame() const { return bOverrideStartFrame; }
+	UFUNCTION(BlueprintPure, Category="Render Pages|Page")
+	bool GetIsUsingCustomStartFrame() const { return bOverrideStartFrame; }
 
-	UFUNCTION(BlueprintCallable, Category="Render Pages | Page")
-	void SetIsCustomStartFrame(const bool bNewOverrideStartFrame) { bOverrideStartFrame = bNewOverrideStartFrame; }
+	UFUNCTION(BlueprintCallable, Category="Render Pages|Page")
+	void SetIsUsingCustomStartFrame(const bool bNewOverrideStartFrame) { bOverrideStartFrame = bNewOverrideStartFrame; }
 
-	UFUNCTION(BlueprintPure, Category="Render Pages | Page")
+	UFUNCTION(BlueprintPure, Category="Render Pages|Page")
 	int32 GetCustomStartFrame() const { return CustomStartFrame; }
 
-	UFUNCTION(BlueprintCallable, Category="Render Pages | Page")
+	UFUNCTION(BlueprintCallable, Category="Render Pages|Page")
 	void SetCustomStartFrame(const int32 NewCustomStartFrame) { CustomStartFrame = NewCustomStartFrame; }
 
 
-	UFUNCTION(BlueprintPure, Category="Render Pages | Page")
-	bool GetIsCustomEndFrame() const { return bOverrideEndFrame; }
+	UFUNCTION(BlueprintPure, Category="Render Pages|Page")
+	bool GetIsUsingCustomEndFrame() const { return bOverrideEndFrame; }
 
-	UFUNCTION(BlueprintCallable, Category="Render Pages | Page")
-	void SetIsCustomEndFrame(const bool bNewOverrideEndFrame) { bOverrideEndFrame = bNewOverrideEndFrame; }
+	UFUNCTION(BlueprintCallable, Category="Render Pages|Page")
+	void SetIsUsingCustomEndFrame(const bool bNewOverrideEndFrame) { bOverrideEndFrame = bNewOverrideEndFrame; }
 
-	UFUNCTION(BlueprintPure, Category="Render Pages | Page")
+	UFUNCTION(BlueprintPure, Category="Render Pages|Page")
 	int32 GetCustomEndFrame() const { return CustomEndFrame; }
 
-	UFUNCTION(BlueprintCallable, Category="Render Pages | Page")
+	UFUNCTION(BlueprintCallable, Category="Render Pages|Page")
 	void SetCustomEndFrame(const int32 NewCustomEndFrame) { CustomEndFrame = NewCustomEndFrame; }
 
 
-	UFUNCTION(BlueprintPure, Category="Render Pages | Page")
-	bool GetIsCustomResolution() const { return bOverrideResolution; }
+	UFUNCTION(BlueprintPure, Category="Render Pages|Page")
+	bool GetIsUsingCustomResolution() const { return bOverrideResolution; }
 
-	UFUNCTION(BlueprintCallable, Category="Render Pages | Page")
-	void SetIsCustomResolution(const bool bNewOverrideResolution) { bOverrideResolution = bNewOverrideResolution; }
+	UFUNCTION(BlueprintCallable, Category="Render Pages|Page")
+	void SetIsUsingCustomResolution(const bool bNewOverrideResolution) { bOverrideResolution = bNewOverrideResolution; }
 
-	UFUNCTION(BlueprintPure, Category="Render Pages | Page")
+	UFUNCTION(BlueprintPure, Category="Render Pages|Page")
 	FIntPoint GetCustomResolution() const { return CustomResolution; }
 
-	UFUNCTION(BlueprintCallable, Category="Render Pages | Page")
+	UFUNCTION(BlueprintCallable, Category="Render Pages|Page")
 	void SetCustomResolution(const FIntPoint NewCustomResolution) { CustomResolution = NewCustomResolution; }
 
 
-	UFUNCTION(BlueprintPure, Category="Render Pages | Page")
+	UFUNCTION(BlueprintPure, Category="Render Pages|Page")
 	FString GetPageId() const { return PageId; }
 
 	static FString PurgePageIdOrReturnEmptyString(const FString& NewPageId);
 	static FString PurgePageId(const FString& NewPageId);
 	static FString PurgePageIdOrGenerateUniqueId(URenderPageCollection* PageCollection, const FString& NewPageId);
 
-	UFUNCTION(BlueprintCallable, Category="Render Pages | Page")
+	UFUNCTION(BlueprintCallable, Category="Render Pages|Page")
 	void SetPageId(const FString& NewPageId) { PageId = PurgePageId(NewPageId); }
 	void SetPageIdRaw(const FString& NewPageId) { PageId = NewPageId; }
 
 
-	UFUNCTION(BlueprintPure, Category="Render Pages | Page")
+	UFUNCTION(BlueprintPure, Category="Render Pages|Page")
 	FString GetPageName() const { return PageName; }
 
 	static FString PurgePageName(const FString& NewPageName);
 
-	UFUNCTION(BlueprintCallable, Category="Render Pages | Page")
+	UFUNCTION(BlueprintCallable, Category="Render Pages|Page")
 	void SetPageName(const FString& NewPageName) { PageName = NewPageName; }
 	void SetPageNameRaw(const FString& NewPageName) { PageName = NewPageName; }
 
 
-	UFUNCTION(BlueprintPure, Category="Render Pages | Page")
+	UFUNCTION(BlueprintPure, Category="Render Pages|Page")
 	bool GetIsEnabled() const { return bIsEnabled; }
 
-	UFUNCTION(BlueprintCallable, Category="Render Pages | Page")
+	UFUNCTION(BlueprintCallable, Category="Render Pages|Page")
 	void SetIsEnabled(const bool bEnabled) { bIsEnabled = bEnabled; }
 
 
-	UFUNCTION(BlueprintPure, Category="Render Pages | Page")
+	UFUNCTION(BlueprintPure, Category="Render Pages|Page")
 	FString GetOutputDirectory() const;
 	FString GetOutputDirectoryRaw() const { return OutputDirectory; }
 
-	UFUNCTION(BlueprintPure, Category="Render Pages | Page")
+	UFUNCTION(BlueprintPure, Category="Render Pages|Page")
 	FString GetOutputDirectoryForDisplay() const { return OutputDirectory; }
 
 	static FString PurgeOutputDirectory(const FString& NewOutputDirectory);
 
-	UFUNCTION(BlueprintCallable, Category="Render Pages | Page")
+	UFUNCTION(BlueprintCallable, Category="Render Pages|Page")
 	void SetOutputDirectory(const FString& NewOutputDirectory) { OutputDirectory = PurgeOutputDirectory(NewOutputDirectory); }
 	void SetOutputDirectoryRaw(const FString& NewOutputDirectory) { OutputDirectory = NewOutputDirectory; }
 
 
-	UFUNCTION(BlueprintPure, Category="Render Pages | Page")
+	UFUNCTION(BlueprintPure, Category="Render Pages|Page")
 	UMoviePipelineMasterConfig* GetRenderPreset() const { return RenderPreset; }
 
-	UFUNCTION(BlueprintPure, Category="Render Pages | Page")
+	UFUNCTION(BlueprintPure, Category="Render Pages|Page")
 	UMoviePipelineOutputSetting* GetRenderPresetOutputSettings() const;
 
-	UFUNCTION(BlueprintCallable, Category="Render Pages | Page")
+	UFUNCTION(BlueprintCallable, Category="Render Pages|Page")
 	void SetRenderPreset(UMoviePipelineMasterConfig* NewRenderPreset) { RenderPreset = NewRenderPreset; }
 
 
@@ -329,10 +329,16 @@ private:
 	void CopyValuesToOrFromCDO(const bool bToCDO);
 
 public:
+	UFUNCTION(BlueprintPure, Category="Render Pages|Page Collection")
 	FGuid GetId() const { return Id; }
+
+	UFUNCTION(BlueprintCallable, Category="Render Pages|Page Collection")
 	void GenerateNewId() { Id = FGuid::NewGuid(); }
 
+	UFUNCTION(BlueprintCallable, Category="Render Pages|Page Collection")
 	void SetPropsSource(ERenderPagePropsSourceType InPropsSourceType, UObject* InPropsSourceOrigin = nullptr);
+
+	UFUNCTION(BlueprintPure, Category="Render Pages|Page Collection")
 	URenderPagePropsSourceBase* GetPropsSource() const;
 
 	template<typename Type>
@@ -342,22 +348,70 @@ public:
 		return Cast<Type>(GetPropsSource());
 	}
 
+	UFUNCTION(BlueprintPure, Category="Render Pages|Page Collection")
 	ERenderPagePropsSourceType GetPropsSourceType() const { return PropsSourceType; }
+
+	UFUNCTION(BlueprintPure, Category="Render Pages|Page Collection")
 	UObject* GetPropsSourceOrigin() const;
 
 public:
+	UFUNCTION(BlueprintCallable, Category="Render Pages|Page Collection")
 	void AddRenderPage(URenderPage* RenderPage);
+
+	UFUNCTION(BlueprintCallable, Category="Render Pages|Page Collection")
 	void RemoveRenderPage(URenderPage* RenderPage);
+
+	UFUNCTION(BlueprintCallable, Category="Render Pages|Page Collection")
 	void InsertRenderPage(URenderPage* RenderPage, int32 Index);
+
+	UFUNCTION(BlueprintPure, Category="Render Pages|Page Collection")
 	bool HasRenderPage(URenderPage* RenderPage) const;
+
+	UFUNCTION(BlueprintPure, Category="Render Pages|Page Collection")
 	int32 GetIndexOfRenderPage(URenderPage* RenderPage) const;
+
 	TArray<TObjectPtr<URenderPage>>& GetRenderPagesRef() { return RenderPages; }
+
+	UFUNCTION(BlueprintPure, Category="Render Pages|Page Collection")
 	TArray<URenderPage*> GetRenderPages() const;
+
+	UFUNCTION(BlueprintPure, Category="Render Pages|Page Collection")
 	TArray<URenderPage*> GetEnabledRenderPages() const;
+
+	UFUNCTION(BlueprintPure, Category="Render Pages|Page Collection")
 	TArray<URenderPage*> GetDisabledRenderPages() const;
+
+	UFUNCTION(BlueprintCallable, Category="Render Pages|Page Collection")
 	void InsertRenderPageBefore(URenderPage* RenderPage, URenderPage* BeforeRenderPage);
+
+	UFUNCTION(BlueprintCallable, Category="Render Pages|Page Collection")
 	void InsertRenderPageAfter(URenderPage* RenderPage, URenderPage* AfterRenderPage);
 
+public:
+	/** Generates a unique page ID by finding the highest page ID and increasing it by one. **/
+	UFUNCTION(BlueprintCallable, Category="Render Pages|Page Collection")
+	FString GenerateNextPageId();
+
+	/** Finds whether given page ID already exists in this collection. **/
+	UFUNCTION(BlueprintPure, Category="Render Pages|Page Collection")
+	bool DoesPageIdExist(const FString& PageId);
+
+	/** Creates a new render page and adds it to this collection. **/
+	UFUNCTION(BlueprintCallable, Category="Render Pages|Page Collection")
+	URenderPage* CreateAndAddNewRenderPage();
+
+	/** Copy the given page in this collection. **/
+	UFUNCTION(BlueprintCallable, Category="Render Pages|Page Collection")
+	URenderPage* DuplicateAndAddRenderPage(URenderPage* Page);
+
+	/** Relocates the given page in this collection to the position of the given dropped-on page. **/
+	UFUNCTION(BlueprintCallable, Category="Render Pages|Page Collection")
+	bool ReorderRenderPage(URenderPage* Page, URenderPage* DroppedOnPage, const bool bAfter = true);
+
+public:
+	/** Returns true when it's currently executing a blueprint implementable event, returns false otherwise. */
+	bool IsCurrentlyExecutingUserCode() const { return bExecutingPreRender || bExecutingPostRender; }
+
 private:
 	DECLARE_MULTICAST_DELEGATE(FOnRenderPageCollectionPreSave);
 	FOnRenderPageCollectionPreSave& OnPreSaveCDO() { return GetCDO()->OnRenderPageCollectionPreSaveDelegate; }
@@ -386,6 +440,15 @@ private:
 	TArray<TObjectPtr<URenderPage>> RenderPages;
 
 
+	/** True when it's currently executing the PreRender event, false otherwise. */
+	UPROPERTY(Transient)
+	bool bExecutingPreRender;
+
+	/** True when it's currently executing the PostRender event, false otherwise. */
+	UPROPERTY(Transient)
+	bool bExecutingPostRender;
+
+
 	/** GetPropsSource calls are somewhat expensive, we speed that up by caching the result (the PropsSource) that has been last outputted by that function. */
 	UPROPERTY(Transient)
 	mutable TObjectPtr<URenderPagePropsSourceBase> CachedPropsSource;
diff --git a/Engine/Plugins/Experimental/RenderPages/Source/RenderPages/Public/RenderPage/RenderPageManager.h b/Engine/Plugins/Experimental/RenderPages/Source/RenderPages/Public/RenderPage/RenderPageManager.h
--- a/Engine/Plugins/Experimental/RenderPages/Source/RenderPages/Public/RenderPage/RenderPageManager.h
+++ b/Engine/Plugins/Experimental/RenderPages/Source/RenderPages/Public/RenderPage/RenderPageManager.h
@@ -87,28 +87,12 @@ namespace UE::RenderPages
 		/** The number of characters for a generated ID. For example, a value of 4 results in IDs: "0001", "0002", etc. */
 		static constexpr int32 GeneratedIdCharacterLength = 4;
 
-	public:
-		/** Creates a new page and adds it to the given collection. **/
-		URenderPage* AddNewPage(URenderPageCollection* PageCollection);
-
-		/** Copy the given pages in the given collection. **/
-		URenderPage* CopyPage(URenderPageCollection* PageCollection, URenderPage* Page);
-
-		/** Delete the given page in the given collection. **/
-		void DeletePage(URenderPageCollection* PageCollection, URenderPage* Page);
-
-		/** Relocates the given page in the given collection to the position of the given dropped-on page. **/
-		bool DragDropPage(URenderPageCollection* PageCollection, URenderPage* Page, URenderPage* DroppedOnPage, const bool bAfter = true);
-
-		/** Generates a unique page ID by finding the highest page ID and increasing it by one. **/
-		FString CreateUniquePageId(URenderPageCollection* PageCollection);
-
-		/** Finds whether given page ID already exists in the collection. **/
-		bool DoesPageIdExist(URenderPageCollection* PageCollection, const FString& PageId);
 
+	public:
 		/** Batch render the current pages of the given collection. **/
 		URenderPagesMoviePipelineRenderJob* CreateBatchRenderJob(URenderPageCollection* PageCollection);
 
+
 		/** Batch render the current pages of the given collection. **/
 		URenderPagesMoviePipelineRenderJob* RenderPreviewFrame(const FRenderPageManagerRenderPreviewFrameArgs& Args);
 
@@ -124,6 +108,7 @@ namespace UE::RenderPages
 		/** Gets the rendered preview frame the given frame number (of a rendering in which the frame number was not specified). */
 		UTexture2D* GetRenderedPreviewFrame(URenderPage* Page, const int32 Frame);
 
+
 		/** Makes sure that all the data from the current props source is stored in all of the pages of this page collection. */
 		void UpdatePagesPropValues(URenderPageCollection* PageCollection);
 
@@ -133,6 +118,7 @@ namespace UE::RenderPages
 		/** Restores the props that were previously applied, to the values they were before. */
 		void RestorePagePropValues(const FRenderPageManagerPreviousPagePropValues& PreviousPropValues);
 
+
 	private:
 		/** The map that stores the start frame (of a render) of each rendered page. */
 		TMap<FGuid, int32> StartFrameOfRenders;
diff --git a/Engine/Plugins/Experimental/RenderPages/Source/RenderPages/Public/RenderPage/RenderPageMoviePipelineJob.h b/Engine/Plugins/Experimental/RenderPages/Source/RenderPages/Public/RenderPage/RenderPageMoviePipelineJob.h
--- a/Engine/Plugins/Experimental/RenderPages/Source/RenderPages/Public/RenderPage/RenderPageMoviePipelineJob.h
+++ b/Engine/Plugins/Experimental/RenderPages/Source/RenderPages/Public/RenderPage/RenderPageMoviePipelineJob.h
@@ -13,6 +13,7 @@ struct FMoviePipelineOutputData;
 class URenderPageCollection;
 class URenderPage;
 class UMoviePipelineExecutorBase;
+class UMoviePipelinePIEExecutor;
 class UMoviePipelineQueue;
 class UMoviePipelineExecutorJob;
 class URenderPagesMoviePipelineRenderJob;
@@ -44,7 +45,7 @@ namespace UE::RenderPages
 		TArray<TObjectPtr<URenderPage>> Pages;
 
 		/** If not null, it will override the MRQ pipeline executor class with this class. */
-		TSubclassOf<UMoviePipelineExecutorBase> PipelineExecutorClass = nullptr;
+		TSubclassOf<UMoviePipelinePIEExecutor> PipelineExecutorClass = nullptr;
 
 		/** The MRQ settings classes to disable (things like Anti-Aliasing, High-Res, etc). */
 		TArray<TSubclassOf<UMoviePipelineSetting>> DisableSettingsClasses;
@@ -100,9 +101,19 @@ public:
 
 private:
 	void ComputePlaybackContext(bool& bOutAllowBinding);
+	void ExecutePageStarted(UMoviePipelineExecutorJob* JobToStart);
+	void ExecutePageFinished(FMoviePipelineOutputData PipelineOutputData);
 	void ExecuteFinished(UMoviePipelineExecutorBase* PipelineExecutor, const bool bSuccess);
 
 protected:
+	/** The render page that will be rendered. */
+	UPROPERTY(Transient)
+	TObjectPtr<URenderPage> Page;
+
+	/** The render page collection that the render page (that will be rendered) belongs to. */
+	UPROPERTY(Transient)
+	TObjectPtr<URenderPageCollection> PageCollection;
+
 	/** The MRQ queue. */
 	UPROPERTY(Transient)
 	TObjectPtr<UMoviePipelineQueue> RenderQueue;
@@ -183,10 +194,6 @@ protected:
 	UPROPERTY(Transient)
 	FRenderPagePreviousEngineFpsSettings PreviousFrameLimitSettings;
 
-	/** True if the queue has previously executed the pre-render event of a page. */
-	UPROPERTY(Transient)
-	bool bRanPreRender;
-
 public:
 	/** A delegate for when the render job is about to start. */
 	UE::RenderPages::FOnRenderPagesMoviePipelineRenderJobStarted& OnExecuteStarted() { return OnExecuteStartedDelegate; }
diff --git a/Engine/Plugins/Experimental/RenderPages/Source/RenderPagesDeveloper/Private/Blueprints/RenderPagesBlueprint.cpp b/Engine/Plugins/Experimental/RenderPages/Source/RenderPagesDeveloper/Private/Blueprints/RenderPagesBlueprint.cpp
--- a/Engine/Plugins/Experimental/RenderPages/Source/RenderPagesDeveloper/Private/Blueprints/RenderPagesBlueprint.cpp
+++ b/Engine/Plugins/Experimental/RenderPages/Source/RenderPagesDeveloper/Private/Blueprints/RenderPagesBlueprint.cpp
@@ -60,4 +60,114 @@ void URenderPagesBlueprint::PostLoad()
 		FBlueprintEditorUtils::AddUbergraphPage(this, NewGraph);
 		LastEditedDocuments.AddUnique(NewGraph);
 	}
+
+	FCoreUObjectDelegates::OnObjectModified.RemoveAll(this);
+	OnChanged().RemoveAll(this);
+	FCoreUObjectDelegates::OnObjectModified.AddUObject(this, &URenderPagesBlueprint::OnPreVariablesChange);
+	OnChanged().AddUObject(this, &URenderPagesBlueprint::OnPostVariablesChange);
+
+	OnPreVariablesChange(this);
+	OnPostVariablesChange(this);
+}
+
+void URenderPagesBlueprint::OnPreVariablesChange(UObject* InObject)
+{
+	if (InObject != this)
+	{
+		return;
+	}
+	LastNewVariables = NewVariables;
+}
+
+void URenderPagesBlueprint::OnPostVariablesChange(UBlueprint* InBlueprint)
+{
+	if (InBlueprint != this)
+	{
+		return;
+	}
+
+	bool bFoundChange = false;
+
+	TMap<FGuid, int32> NewVariablesByGuid;
+	for (int32 VarIndex = 0; VarIndex < NewVariables.Num(); VarIndex++)
+	{
+		NewVariablesByGuid.Add(NewVariables[VarIndex].VarGuid, VarIndex);
+	}
+
+	TMap<FGuid, int32> OldVariablesByGuid;
+	for (int32 VarIndex = 0; VarIndex < LastNewVariables.Num(); VarIndex++)
+	{
+		OldVariablesByGuid.Add(LastNewVariables[VarIndex].VarGuid, VarIndex);
+	}
+
+	for (FBPVariableDescription& OldVariable : LastNewVariables)
+	{
+		if (!NewVariablesByGuid.Contains(OldVariable.VarGuid))
+		{
+			bFoundChange = true;
+			OnVariableRemoved(OldVariable);
+		}
+	}
+
+	for (FBPVariableDescription& NewVariable : NewVariables)
+	{
+		if (!OldVariablesByGuid.Contains(NewVariable.VarGuid))
+		{
+			bFoundChange = true;
+			OnVariableAdded(NewVariable);
+			continue;
+		}
+
+		int32 OldVarIndex = OldVariablesByGuid.FindChecked(NewVariable.VarGuid);
+		const FBPVariableDescription& OldVariable = LastNewVariables[OldVarIndex];
+		if (OldVariable.VarName != NewVariable.VarName)
+		{
+			bFoundChange = true;
+			OnVariableRenamed(NewVariable, OldVariable.VarName, NewVariable.VarName);
+		}
+
+		if (OldVariable.VarType != NewVariable.VarType)
+		{
+			bFoundChange = true;
+			OnVariableTypeChanged(NewVariable, OldVariable.VarType, NewVariable.VarType);
+		}
+
+		if (OldVariable.PropertyFlags != NewVariable.PropertyFlags)
+		{
+			bFoundChange = true;
+			OnVariablePropertyFlagsChanged(NewVariable, OldVariable.PropertyFlags, NewVariable.PropertyFlags);
+		}
+	}
+
+	LastNewVariables = NewVariables;
+
+	if (bFoundChange)
+	{
+		FBlueprintEditorUtils::MarkBlueprintAsStructurallyModified(this);
+	}
+}
+
+void URenderPagesBlueprint::OnVariableAdded(FBPVariableDescription& InVar)
+{
+	MakeVariableTransientUnlessInstanceEditable(InVar);
+}
+
+void URenderPagesBlueprint::OnVariablePropertyFlagsChanged(FBPVariableDescription& InVar, const uint64 InOldVarPropertyFlags, const uint64 InNewVarPropertyFlags)
+{
+	if ((InOldVarPropertyFlags & CPF_DisableEditOnInstance) != (InNewVarPropertyFlags & CPF_DisableEditOnInstance))// if the value of [Instance Editable] changed
+	{
+		MakeVariableTransientUnlessInstanceEditable(InVar);
+	}
+}
+
+void URenderPagesBlueprint::MakeVariableTransientUnlessInstanceEditable(FBPVariableDescription& InVar)
+{
+	if ((InVar.PropertyFlags & CPF_DisableEditOnInstance) == 0)// if [Instance Editable]
+	{
+		InVar.PropertyFlags &= ~CPF_Transient;// set [Transient] to false
+	}
+	else
+	{
+		InVar.PropertyFlags |= CPF_Transient;// set [Transient] to true
+	}
 }
diff --git a/Engine/Plugins/Experimental/RenderPages/Source/RenderPagesDeveloper/Public/Blueprints/RenderPagesBlueprint.h b/Engine/Plugins/Experimental/RenderPages/Source/RenderPagesDeveloper/Public/Blueprints/RenderPagesBlueprint.h
--- a/Engine/Plugins/Experimental/RenderPages/Source/RenderPagesDeveloper/Public/Blueprints/RenderPagesBlueprint.h
+++ b/Engine/Plugins/Experimental/RenderPages/Source/RenderPagesDeveloper/Public/Blueprints/RenderPagesBlueprint.h
@@ -19,7 +19,6 @@ class RENDERPAGESDEVELOPER_API URenderPagesBlueprint : public UBlueprint
 
 public:
 	//~ Begin UBlueprint Interface
-	virtual UClass* GetBlueprintClass() const override;
 	virtual bool SupportedByDefaultBlueprintFactory() const override { return false; }
 	virtual bool IsValidForBytecodeOnlyRecompile() const override { return false; }
 	virtual bool SupportsGlobalVariables() const override { return true; }
@@ -30,6 +29,24 @@ public:
 	virtual bool SupportsEventGraphs() const override { return true; }
 	virtual bool SupportsAnimLayers() const override { return false; }
 
+	virtual UClass* GetBlueprintClass() const override;
 	virtual void PostLoad() override;
 	//~ End UBlueprint Interface
+
+private:
+	virtual void OnPreVariablesChange(UObject* InObject);
+	virtual void OnPostVariablesChange(UBlueprint* InBlueprint);
+
+protected:
+	virtual void OnVariableAdded(FBPVariableDescription& InVar);
+	virtual void OnVariableRemoved(FBPVariableDescription& InVar) {}
+	virtual void OnVariableRenamed(FBPVariableDescription& InVar, const FName& InOldVarName, const FName& InNewVarName) {}
+	virtual void OnVariableTypeChanged(FBPVariableDescription& InVar, const FEdGraphPinType& InOldVarType, const FEdGraphPinType& InNewVarType) {}
+	virtual void OnVariablePropertyFlagsChanged(FBPVariableDescription& InVar, const uint64 InOldVarPropertyFlags, const uint64 InNewVarPropertyFlags);
+
+protected:
+	virtual void MakeVariableTransientUnlessInstanceEditable(FBPVariableDescription& InVar);
+
+private:
+	TArray<FBPVariableDescription> LastNewVariables;
 };
diff --git a/Engine/Plugins/Experimental/RenderPages/Source/RenderPagesEditor/Private/Toolkit/RenderPageCollectionEditor.cpp b/Engine/Plugins/Experimental/RenderPages/Source/RenderPagesEditor/Private/Toolkit/RenderPageCollectionEditor.cpp
--- a/Engine/Plugins/Experimental/RenderPages/Source/RenderPagesEditor/Private/Toolkit/RenderPageCollectionEditor.cpp
+++ b/Engine/Plugins/Experimental/RenderPages/Source/RenderPagesEditor/Private/Toolkit/RenderPageCollectionEditor.cpp
@@ -352,8 +352,10 @@ void UE::RenderPages::Private::FRenderPageCollectionEditor::AddReferencedObjects
 {
 	FBlueprintEditor::AddReferencedObjects(Collector);
 
-	URenderPageCollection* RenderPageCollection = GetInstance();
-	Collector.AddReferencedObject(RenderPageCollection);
+	if (URenderPageCollection* Instance = GetInstance(); IsValid(Instance))
+	{
+		Collector.AddReferencedObject(Instance);
+	}
 }
 
 void UE::RenderPages::Private::FRenderPageCollectionEditor::NotifyPreChange(FProperty* PropertyAboutToChange)
@@ -385,15 +387,15 @@ void UE::RenderPages::Private::FRenderPageCollectionEditor::BindCommands()
 
 void UE::RenderPages::Private::FRenderPageCollectionEditor::AddPageAction()
 {
-	URenderPage* Page = IRenderPagesModule::Get().GetManager().AddNewPage(GetInstance());
-	if (!IsValid(Page))
+	if (URenderPageCollection* Instance = GetInstance(); IsValid(Instance))
 	{
-		return;
+		if (URenderPage* Page = Instance->CreateAndAddNewRenderPage(); IsValid(Page))
+		{
+			OnRenderPageCreated().Broadcast(Page);
+			MarkAsModified();
+			OnRenderPagesChanged().Broadcast();
+		}
 	}
-
-	OnRenderPageCreated().Broadcast(Page);
-	MarkAsModified();
-	OnRenderPagesChanged().Broadcast();
 }
 
 void UE::RenderPages::Private::FRenderPageCollectionEditor::CopyPageAction()
@@ -404,14 +406,15 @@ void UE::RenderPages::Private::FRenderPageCollectionEditor::CopyPageAction()
 		return;
 	}
 
-	URenderPageCollection* RenderPageCollection = GetInstance();
-	for (URenderPage* SelectedRenderPage : SelectedRenderPages)
+	if (URenderPageCollection* Instance = GetInstance(); IsValid(Instance))
 	{
-		IRenderPagesModule::Get().GetManager().CopyPage(RenderPageCollection, SelectedRenderPage);
+		for (URenderPage* SelectedRenderPage : SelectedRenderPages)
+		{
+			Instance->DuplicateAndAddRenderPage(SelectedRenderPage);
+		}
+		MarkAsModified();
+		OnRenderPagesChanged().Broadcast();
 	}
-
-	MarkAsModified();
-	OnRenderPagesChanged().Broadcast();
 }
 
 void UE::RenderPages::Private::FRenderPageCollectionEditor::DeletePageAction()
@@ -433,14 +436,16 @@ void UE::RenderPages::Private::FRenderPageCollectionEditor::DeletePageAction()
 		return;
 	}
 
-	URenderPageCollection* RenderPageCollection = GetInstance();
-	for (URenderPage* SelectedRenderPage : SelectedRenderPages)
+	if (URenderPageCollection* Instance = GetInstance(); IsValid(Instance))
 	{
-		IRenderPagesModule::Get().GetManager().DeletePage(RenderPageCollection, SelectedRenderPage);
-	}
+		for (URenderPage* SelectedRenderPage : SelectedRenderPages)
+		{
+			Instance->RemoveRenderPage(SelectedRenderPage);
+		}
 
-	MarkAsModified();
-	OnRenderPagesChanged().Broadcast();
+		MarkAsModified();
+		OnRenderPagesChanged().Broadcast();
+	}
 }
 
 void UE::RenderPages::Private::FRenderPageCollectionEditor::BatchRenderListAction()
@@ -451,26 +456,26 @@ void UE::RenderPages::Private::FRenderPageCollectionEditor::BatchRenderListActio
 		return;
 	}
 
-	URenderPageCollection* PageCollection = GetInstance();
-	if (PageCollection->GetRenderPages().Num() <= 0)
+	if (URenderPageCollection* Instance = GetInstance(); IsValid(Instance))
 	{
-		const FText TitleText = LOCTEXT("NoPagesToRenderTitle", "No Pages To Render");
-		FMessageDialog::Open(
-			EAppMsgType::Ok,
-			LOCTEXT("NoPagesToRenderText", "There are no render pages in this collection, and so nothing can be rendered. Please make a page and try again."),
-			&TitleText);
-		return;
-	}
+		if (Instance->GetRenderPages().Num() <= 0)
+		{
+			const FText TitleText = LOCTEXT("NoPagesToRenderTitle", "No Pages To Render");
+			FMessageDialog::Open(
+				EAppMsgType::Ok,
+				LOCTEXT("NoPagesToRenderText", "There are no render pages in this collection, and so nothing can be rendered. Please make a page and try again."),
+				&TitleText);
+			return;
+		}
 
-	URenderPagesMoviePipelineRenderJob* RenderJob = IRenderPagesModule::Get().GetManager().CreateBatchRenderJob(PageCollection);
-	if (!RenderJob)
-	{
-		return;
+		if (URenderPagesMoviePipelineRenderJob* RenderJob = IRenderPagesModule::Get().GetManager().CreateBatchRenderJob(Instance); IsValid(RenderJob))
+		{
+			RenderJob->OnExecuteFinished().AddRaw(this, &FRenderPageCollectionEditor::OnBatchRenderListActionFinished);
+			BatchRenderJob = RenderJob;
+			OnRenderPagesBatchRenderingStarted().Broadcast(RenderJob);
+			RenderJob->Execute();
+		}
 	}
-	RenderJob->OnExecuteFinished().AddRaw(this, &FRenderPageCollectionEditor::OnBatchRenderListActionFinished);
-	BatchRenderJob = RenderJob;
-	OnRenderPagesBatchRenderingStarted().Broadcast(RenderJob);
-	RenderJob->Execute();
 }
 
 void UE::RenderPages::Private::FRenderPageCollectionEditor::OnBatchRenderListActionFinished(URenderPagesMoviePipelineRenderJob* RenderJob, bool bSuccess)
@@ -542,11 +547,11 @@ void UE::RenderPages::Private::FRenderPageCollectionEditor::FillToolbar(FToolBar
 
 void UE::RenderPages::Private::FRenderPageCollectionEditor::DestroyInstance()
 {
-	if (URenderPageCollection* RenderPageCollection = GetInstance())
+	if (URenderPageCollection* Instance = GetInstance(); IsValid(Instance))
 	{
-		RenderPageCollection->OnClose();
+		Instance->OnClose();
 		RenderPageCollectionWeakPtr.Reset();
-		RenderPageCollection->MarkAsGarbage();
+		Instance->MarkAsGarbage();
 	}
 }
 
diff --git a/Engine/Plugins/Experimental/RenderPages/Source/RenderPagesEditor/Private/UI/SRenderPagesPageList.cpp b/Engine/Plugins/Experimental/RenderPages/Source/RenderPagesEditor/Private/UI/SRenderPagesPageList.cpp
--- a/Engine/Plugins/Experimental/RenderPages/Source/RenderPagesEditor/Private/UI/SRenderPagesPageList.cpp
+++ b/Engine/Plugins/Experimental/RenderPages/Source/RenderPagesEditor/Private/UI/SRenderPagesPageList.cpp
@@ -355,11 +355,14 @@ FReply UE::RenderPages::Private::SRenderPagesPageListTableRow::OnAcceptDrop(cons
 	{
 		if (const TSharedPtr<FRenderPagesPageListTableRowDragDropOp> DragDropOp = Event.GetOperationAs<FRenderPagesPageListTableRowDragDropOp>())
 		{
-			if (IRenderPagesModule::Get().GetManager().DragDropPage(BlueprintEditor->GetInstance(), DragDropOp->GetPage(), Page, (InItemDropZone != EItemDropZone::AboveItem)))
+			if (URenderPageCollection* Instance = BlueprintEditor->GetInstance(); IsValid(Instance))
 			{
-				BlueprintEditor->MarkAsModified();
-				BlueprintEditor->OnRenderPagesChanged().Broadcast();
-				return FReply::Handled();
+				if (Instance->ReorderRenderPage(DragDropOp->GetPage(), Page, (InItemDropZone != EItemDropZone::AboveItem)))
+				{
+					BlueprintEditor->MarkAsModified();
+					BlueprintEditor->OnRenderPagesChanged().Broadcast();
+					return FReply::Handled();
+				}
 			}
 		}
 	}
@@ -427,7 +430,7 @@ TSharedRef<SWidget> UE::RenderPages::Private::SRenderPagesPageListTableRow::Gene
 				{
 					if (URenderPageCollection* PageCollection = BlueprintEditor->GetInstance(); IsValid(PageCollection))
 					{
-						if (IRenderPagesModule::Get().GetManager().DoesPageIdExist(PageCollection, NewPageId))
+						if (PageCollection->DoesPageIdExist(NewPageId))
 						{
 							const FText TitleText = LOCTEXT("PageIdNotUniqueTitle", "Duplicate Page IDs");
 							FMessageDialog::Open(
diff --git a/Engine/Plugins/Experimental/RenderPages/Source/RenderPagesEditor/Private/UI/SRenderPagesPageViewerLive.cpp b/Engine/Plugins/Experimental/RenderPages/Source/RenderPagesEditor/Private/UI/SRenderPagesPageViewerLive.cpp
--- a/Engine/Plugins/Experimental/RenderPages/Source/RenderPagesEditor/Private/UI/SRenderPagesPageViewerLive.cpp
+++ b/Engine/Plugins/Experimental/RenderPages/Source/RenderPagesEditor/Private/UI/SRenderPagesPageViewerLive.cpp
@@ -35,43 +35,64 @@ void UE::RenderPages::Private::SRenderPagesEditorViewport::Tick(const FGeometry&
 	{
 		return;
 	}
-	ViewportClient->Invalidate();// causes it to rerender
 
-	if (ULevelSequencePlayer* SequencePlayer = GetSequencePlayer(); IsValid(SequencePlayer))
+	ULevelSequencePlayer* SequencePlayer = GetSequencePlayer();
+	if (IsValid(SequencePlayer))
 	{
 		SequencePlayer->SetPlaybackPosition(FMovieSceneSequencePlaybackParams(LevelSequenceTime, EUpdatePositionMethod::Play));// execute this every tick, in case any sequencer values get overwritten (by remote control props for example)
+	}
+
+	if (IsValid(Page) && IsValid(PageCollection))
+	{
+		PageCollection->PreRender(Page);
+	}
+
+	bool bHasSetCameraData = false;
+	if (IsValid(SequencePlayer))
+	{
 		if (UCameraComponent* Camera = SequencePlayer->GetActiveCameraComponent(); IsValid(Camera))
 		{
 			ViewportClient->SetViewLocation(Camera->GetComponentLocation());
 			ViewportClient->SetViewRotation(Camera->GetComponentRotation());
 			ViewportClient->ViewFOV = Camera->FieldOfView;
-			return;
+			bHasSetCameraData = true;
 		}
 	}
-	ViewportClient->ViewFOV = 90;
-
-	if (UWorld* World = GetWorld(); IsValid(World))
+	if (!bHasSetCameraData)
 	{
-		if (APlayerController* LocalPlayerController = World->GetFirstPlayerController(); IsValid(LocalPlayerController))
-		{
-			FVector ViewLocation;
-			FRotator ViewRotation;
-			LocalPlayerController->GetPlayerViewPoint(ViewLocation, ViewRotation);
-			ViewportClient->SetViewLocation(ViewLocation);
-			ViewportClient->SetViewRotation(ViewRotation);
-			return;
-		}
-
-		for (TActorIterator<APlayerStart> It(World); It; ++It)
+		if (UWorld* World = GetWorld(); IsValid(World))
 		{
-			if (APlayerStart* PlayerStart = *It; IsValid(PlayerStart))
+			if (APlayerController* LocalPlayerController = World->GetFirstPlayerController(); IsValid(LocalPlayerController))
 			{
-				ViewportClient->SetViewLocation(PlayerStart->GetActorLocation());
-				ViewportClient->SetViewRotation(PlayerStart->GetActorRotation());
-				return;
+				FVector ViewLocation;
+				FRotator ViewRotation;
+				LocalPlayerController->GetPlayerViewPoint(ViewLocation, ViewRotation);
+				ViewportClient->SetViewLocation(ViewLocation);
+				ViewportClient->SetViewRotation(ViewRotation);
+				ViewportClient->ViewFOV = 90;
+			}
+			else
+			{
+				for (TActorIterator<APlayerStart> It(World); It; ++It)
+				{
+					if (APlayerStart* PlayerStart = *It; IsValid(PlayerStart))
+					{
+						ViewportClient->SetViewLocation(PlayerStart->GetActorLocation());
+						ViewportClient->SetViewRotation(PlayerStart->GetActorRotation());
+						ViewportClient->ViewFOV = 90;
+						break;
+					}
+				}
 			}
 		}
 	}
+
+	ViewportClient->Viewport->Draw();
+
+	if (IsValid(Page) && IsValid(PageCollection))
+	{
+		PageCollection->PostRender(Page);
+	}
 }
 
 void UE::RenderPages::Private::SRenderPagesEditorViewport::Construct(const FArguments& InArgs)
@@ -82,6 +103,8 @@ void UE::RenderPages::Private::SRenderPagesEditorViewport::Construct(const FArgu
 	LevelSequencePlayer = nullptr;
 	LevelSequence = nullptr;
 	LevelSequenceTime = 0.0f;
+	Page = nullptr;
+	PageCollection = nullptr;
 
 	SEditorViewport::Construct(SEditorViewport::FArguments());
 }
@@ -92,8 +115,10 @@ UE::RenderPages::Private::SRenderPagesEditorViewport::~SRenderPagesEditorViewpor
 	ViewportClient.Reset();
 }
 
-bool UE::RenderPages::Private::SRenderPagesEditorViewport::ShowSequenceFrame(ULevelSequence* InSequence, const float InTime)
+bool UE::RenderPages::Private::SRenderPagesEditorViewport::ShowSequenceFrame(URenderPage* InPage, URenderPageCollection* InPageCollection, ULevelSequence* InSequence, const float InTime)
 {
+	Page = InPage;
+	PageCollection = InPageCollection;
 	LevelSequenceTime = InTime;
 	if (!IsValid(InSequence))
 	{
@@ -293,19 +318,25 @@ void UE::RenderPages::Private::SRenderPagesPageViewerLive::UpdateViewport()
 		{
 			if (ULevelSequence* Sequence = SelectedPage->GetSequence(); IsValid(Sequence))
 			{
-				TOptional<double> StartTime = SelectedPage->GetStartTime();
-				TOptional<double> EndTime = SelectedPage->GetEndTime();
-				if (!StartTime.IsSet() || !EndTime.IsSet() || (StartTime.Get(0) > EndTime.Get(0)))
-				{
-					return;
-				}
-
-				if (!ViewportWidget->ShowSequenceFrame(Sequence, FMath::Lerp(StartTime.Get(0), EndTime.Get(0), FrameSlider->GetValue())))
+				if (const TSharedPtr<IRenderPageCollectionEditor> BlueprintEditor = BlueprintEditorWeakPtr.Pin())
 				{
-					return;
+					if (URenderPageCollection* PageCollection = BlueprintEditor->GetInstance(); IsValid(PageCollection))
+					{
+						TOptional<double> StartTime = SelectedPage->GetStartTime();
+						TOptional<double> EndTime = SelectedPage->GetEndTime();
+						if (!StartTime.IsSet() || !EndTime.IsSet() || (StartTime.Get(0) > EndTime.Get(0)))
+						{
+							return;
+						}
+
+						if (!ViewportWidget->ShowSequenceFrame(SelectedPage, PageCollection, Sequence, FMath::Lerp(StartTime.Get(0), EndTime.Get(0), FrameSlider->GetValue())))
+						{
+							return;
+						}
+
+						ViewportWidget->SetVisibility(EVisibility::Visible);
+					}
 				}
-
-				ViewportWidget->SetVisibility(EVisibility::Visible);
 			}
 		}
 	}
diff --git a/Engine/Plugins/Experimental/RenderPages/Source/RenderPagesEditor/Private/UI/SRenderPagesPageViewerLive.h b/Engine/Plugins/Experimental/RenderPages/Source/RenderPagesEditor/Private/UI/SRenderPagesPageViewerLive.h
--- a/Engine/Plugins/Experimental/RenderPages/Source/RenderPagesEditor/Private/UI/SRenderPagesPageViewerLive.h
+++ b/Engine/Plugins/Experimental/RenderPages/Source/RenderPagesEditor/Private/UI/SRenderPagesPageViewerLive.h
@@ -12,6 +12,7 @@
 class ALevelSequenceActor;
 class SSlider;
 class URenderPage;
+class URenderPageCollection;
 
 namespace UE::RenderPages
 {
@@ -54,7 +55,7 @@ namespace UE::RenderPages::Private
 		void Construct(const FArguments& InArgs);
 		virtual ~SRenderPagesEditorViewport() override;
 
-		bool ShowSequenceFrame(ULevelSequence* InSequence, const float InTime);
+		bool ShowSequenceFrame(URenderPage* InPage, URenderPageCollection* InPageCollection, ULevelSequence* InSequence, const float InTime);
 
 	protected:
 		//~ Begin SEditorViewport Interface
@@ -88,6 +89,14 @@ namespace UE::RenderPages::Private
 
 		/** The time of the currently playing sequence. */
 		float LevelSequenceTime;
+
+		/** The page that's currently being shown. */
+		UPROPERTY()
+		TObjectPtr<URenderPage> Page;
+
+		/** The collection of the page that's currently being shown. */
+		UPROPERTY()
+		TObjectPtr<URenderPageCollection> PageCollection;
 	};
 
 
diff --git a/Engine/Plugins/Experimental/RenderPages/Source/RenderPagesEditor/Private/UI/SRenderPagesPropsRemoteControl.cpp b/Engine/Plugins/Experimental/RenderPages/Source/RenderPagesEditor/Private/UI/SRenderPagesPropsRemoteControl.cpp
--- a/Engine/Plugins/Experimental/RenderPages/Source/RenderPagesEditor/Private/UI/SRenderPagesPropsRemoteControl.cpp
+++ b/Engine/Plugins/Experimental/RenderPages/Source/RenderPagesEditor/Private/UI/SRenderPagesPropsRemoteControl.cpp
@@ -143,39 +143,47 @@ void UE::RenderPages::Private::SRenderPagesPropsRemoteControl::OnRemoteControlEx
 			return;
 		}
 
-		bool bModified = false;
-		for (const FGuid& Id : ModifiedProperties)
+		if (URenderPageCollection* Collection = BlueprintEditor->GetInstance(); IsValid(Collection))
 		{
-			if (const TSharedPtr<FRemoteControlEntity> Entity = Preset->GetExposedEntity<FRemoteControlEntity>(Id).Pin())
+			if (Collection->IsCurrentlyExecutingUserCode())
 			{
-				TArray<uint8> BinaryArray;
-				if (!URenderPagePropRemoteControl::GetValueOfEntity(Entity, BinaryArray))
-				{
-					continue;
-				}
-				TArray<uint8> StoredBinaryArray;
-				if (!GetSelectedPageFieldValue(Entity, StoredBinaryArray))
-				{
-					continue;
-				}
-				if (BinaryArray == StoredBinaryArray)
-				{
-					continue;
-				}
+				return;
+			}
 
-				if (!SetSelectedPageFieldValue(Entity, BinaryArray))
+			bool bModified = false;
+			for (const FGuid& Id : ModifiedProperties)
+			{
+				if (const TSharedPtr<FRemoteControlEntity> Entity = Preset->GetExposedEntity<FRemoteControlEntity>(Id).Pin())
 				{
-					continue;
+					TArray<uint8> BinaryArray;
+					if (!URenderPagePropRemoteControl::GetValueOfEntity(Entity, BinaryArray))
+					{
+						continue;
+					}
+					TArray<uint8> StoredBinaryArray;
+					if (!GetSelectedPageFieldValue(Entity, StoredBinaryArray))
+					{
+						continue;
+					}
+					if (BinaryArray == StoredBinaryArray)
+					{
+						continue;
+					}
+
+					if (!SetSelectedPageFieldValue(Entity, BinaryArray))
+					{
+						continue;
+					}
+					bModified = true;
 				}
-				bModified = true;
 			}
-		}
 
-		if (bModified)
-		{
-			BlueprintEditor->MarkAsModified();
-			BlueprintEditor->OnRenderPagesChanged().Broadcast();
-			Refresh();
+			if (bModified)
+			{
+				BlueprintEditor->MarkAsModified();
+				BlueprintEditor->OnRenderPagesChanged().Broadcast();
+				Refresh();
+			}
 		}
 	}
 }
