- Made the property panels be able to be scrolled through when they get too large.
- Made it no longer automatically revert remote control properties after the PostRender event (remote control OnExposedPropertiesModified delegate doesn't fire right away, causing issues with it overriding any changes in values the user makes through the UI).

diff --git a/Engine/Plugins/Experimental/RenderPages/Source/RenderPages/Private/RenderPage/RenderPageCollection.cpp b/Engine/Plugins/Experimental/RenderPages/Source/RenderPages/Private/RenderPage/RenderPageCollection.cpp
--- a/Engine/Plugins/Experimental/RenderPages/Source/RenderPages/Private/RenderPage/RenderPageCollection.cpp
+++ b/Engine/Plugins/Experimental/RenderPages/Source/RenderPages/Private/RenderPage/RenderPageCollection.cpp
@@ -8,7 +8,6 @@
 #include "MoviePipelineMasterConfig.h"
 #include "MoviePipelineOutputSetting.h"
 #include "MovieScene.h"
-#include "RenderPagesModule.h"
 #include "UObject/ObjectSaveContext.h"
 
 
@@ -528,7 +527,6 @@ void URenderPageCollection::PostRender(URenderPage* Page)
 {
 	bExecutingPostRender = true;
 	ReceivePostRender(Page);
-	UE::RenderPages::Private::FRenderPagesModule::Get().GetManager().ApplyPagePropValues(this, Page);
 	bExecutingPostRender = false;
 }
 
diff --git a/Engine/Plugins/Experimental/RenderPages/Source/RenderPagesEditor/Private/UI/SRenderPagesCollection.cpp b/Engine/Plugins/Experimental/RenderPages/Source/RenderPagesEditor/Private/UI/SRenderPagesCollection.cpp
--- a/Engine/Plugins/Experimental/RenderPages/Source/RenderPagesEditor/Private/UI/SRenderPagesCollection.cpp
+++ b/Engine/Plugins/Experimental/RenderPages/Source/RenderPagesEditor/Private/UI/SRenderPagesCollection.cpp
@@ -1,12 +1,10 @@
 // Copyright Epic Games, Inc. All Rights Reserved.
 
 #include "UI/SRenderPagesCollection.h"
-#include "Blueprints/RenderPagesBlueprint.h"
 #include "IRenderPageCollectionEditor.h"
 #include "RenderPage/RenderPageCollection.h"
 #include "Kismet2/BlueprintEditorUtils.h"
 #include "SlateOptMacros.h"
-#include "Widgets/Layout/SScrollBox.h"
 
 #define LOCTEXT_NAMESPACE "SRenderPagesCollection"
 
diff --git a/Engine/Plugins/Experimental/RenderPages/Source/RenderPagesEditor/Private/UI/Tabs/SRenderPagesCollectionPropertiesTab.cpp b/Engine/Plugins/Experimental/RenderPages/Source/RenderPagesEditor/Private/UI/Tabs/SRenderPagesCollectionPropertiesTab.cpp
--- a/Engine/Plugins/Experimental/RenderPages/Source/RenderPagesEditor/Private/UI/Tabs/SRenderPagesCollectionPropertiesTab.cpp
+++ b/Engine/Plugins/Experimental/RenderPages/Source/RenderPagesEditor/Private/UI/Tabs/SRenderPagesCollectionPropertiesTab.cpp
@@ -4,6 +4,7 @@
 #include "UI/SRenderPagesCollection.h"
 #include "IRenderPageCollectionEditor.h"
 #include "SlateOptMacros.h"
+#include "Widgets/Layout/SScrollBox.h"
 
 #define LOCTEXT_NAMESPACE "SRenderPagesCollectionPropertiesTab"
 
@@ -14,7 +15,20 @@ void UE::RenderPages::Private::SRenderPagesCollectionPropertiesTab::Construct(co
 {
 	ChildSlot
 	[
-		SNew(SRenderPagesCollection, InBlueprintEditor)
+		SNew(SScrollBox)
+		.Style(FAppStyle::Get(), "ScrollBox")
+
+		+ SScrollBox::Slot()
+		[
+			SNew(SVerticalBox)
+
+			+ SVerticalBox::Slot()
+			.AutoHeight()
+			.Padding(0.0f)
+			[
+				SNew(SRenderPagesCollection, InBlueprintEditor)
+			]
+		]
 	];
 }
 
diff --git a/Engine/Plugins/Experimental/RenderPages/Source/RenderPagesEditor/Private/UI/Tabs/SRenderPagesPagePropertiesTab.cpp b/Engine/Plugins/Experimental/RenderPages/Source/RenderPagesEditor/Private/UI/Tabs/SRenderPagesPagePropertiesTab.cpp
--- a/Engine/Plugins/Experimental/RenderPages/Source/RenderPagesEditor/Private/UI/Tabs/SRenderPagesPagePropertiesTab.cpp
+++ b/Engine/Plugins/Experimental/RenderPages/Source/RenderPagesEditor/Private/UI/Tabs/SRenderPagesPagePropertiesTab.cpp
@@ -5,6 +5,7 @@
 #include "UI/SRenderPagesProps.h"
 #include "IRenderPageCollectionEditor.h"
 #include "SlateOptMacros.h"
+#include "Widgets/Layout/SScrollBox.h"
 
 #define LOCTEXT_NAMESPACE "SRenderPagesPagePropertiesTab"
 
@@ -15,20 +16,26 @@ void UE::RenderPages::Private::SRenderPagesPagePropertiesTab::Construct(const FA
 {
 	ChildSlot
 	[
-		SNew(SVerticalBox)
+		SNew(SScrollBox)
+		.Style(FAppStyle::Get(), "ScrollBox")
 
-		+ SVerticalBox::Slot()
-		.AutoHeight()
-		.Padding(0.0f)
+		+ SScrollBox::Slot()
 		[
-			SNew(SRenderPagesPage, InBlueprintEditor)
-		]
-
-		+ SVerticalBox::Slot()
-		.AutoHeight()
-		.Padding(0.0f)
-		[
-			SNew(SRenderPagesProps, InBlueprintEditor)
+			SNew(SVerticalBox)
+
+			+ SVerticalBox::Slot()
+			.AutoHeight()
+			.Padding(0.0f)
+			[
+				SNew(SRenderPagesPage, InBlueprintEditor)
+			]
+
+			+ SVerticalBox::Slot()
+			.AutoHeight()
+			.Padding(0.0f)
+			[
+				SNew(SRenderPagesProps, InBlueprintEditor)
+			]
 		]
 	];
 }
