- Made it impossible for the live viewer to render a page at the same time as one of the other previewers/renderers.
- Made it so that a PreRender event always gets followed by a PostRender event before the PreRender event is able to be fired again.
- Made the viewers show text during waiting (and for the Rendered viewer, it will show a message basically saying "click here").

diff --git a/Engine/Plugins/Experimental/RenderPages/Source/RenderPagesEditor/Private/Toolkit/RenderPageCollectionEditor.cpp b/Engine/Plugins/Experimental/RenderPages/Source/RenderPagesEditor/Private/Toolkit/RenderPageCollectionEditor.cpp
--- a/Engine/Plugins/Experimental/RenderPages/Source/RenderPagesEditor/Private/Toolkit/RenderPageCollectionEditor.cpp
+++ b/Engine/Plugins/Experimental/RenderPages/Source/RenderPagesEditor/Private/Toolkit/RenderPageCollectionEditor.cpp
@@ -95,14 +95,12 @@ UE::RenderPages::Private::FRenderPageCollectionEditor::FRenderPageCollectionEdit
 UE::RenderPages::Private::FRenderPageCollectionEditor::~FRenderPageCollectionEditor()
 {
 	URenderPagesBlueprint* RenderPagesBlueprint = FRenderPageCollectionEditor::GetRenderPagesBlueprint();
-
-	RenderPagesEditorClosedDelegate.Broadcast(this, RenderPagesBlueprint);
-
-	if (RenderPagesBlueprint)
+	if (IsValid(RenderPagesBlueprint))
 	{
-		RenderPagesBlueprint->SetObjectBeingDebugged(nullptr);
-		RenderPagesBlueprint = nullptr;
+		RenderPagesEditorClosedDelegate.Broadcast(this, RenderPagesBlueprint);
 	}
+	RenderPagesBlueprint = nullptr;
+
 	DestroyInstance();
 }
 
diff --git a/Engine/Plugins/Experimental/RenderPages/Source/RenderPagesEditor/Private/UI/SRenderPagesPageViewerLive.cpp b/Engine/Plugins/Experimental/RenderPages/Source/RenderPagesEditor/Private/UI/SRenderPagesPageViewerLive.cpp
--- a/Engine/Plugins/Experimental/RenderPages/Source/RenderPagesEditor/Private/UI/SRenderPagesPageViewerLive.cpp
+++ b/Engine/Plugins/Experimental/RenderPages/Source/RenderPagesEditor/Private/UI/SRenderPagesPageViewerLive.cpp
@@ -31,94 +31,118 @@ UE::RenderPages::Private::FRenderPagesEditorViewportClient::FRenderPagesEditorVi
 
 void UE::RenderPages::Private::SRenderPagesEditorViewport::Tick(const FGeometry&, const double, const float)
 {
-	if (!ViewportClient.IsValid())
-	{
-		return;
-	}
+	Render();
+}
 
-	ULevelSequencePlayer* SequencePlayer = GetSequencePlayer();
-	if (IsValid(SequencePlayer))
-	{
-		SequencePlayer->SetPlaybackPosition(FMovieSceneSequencePlaybackParams(LevelSequenceTime, EUpdatePositionMethod::Play));// execute this every tick, in case any sequencer values get overwritten (by remote control props for example)
-	}
+void UE::RenderPages::Private::SRenderPagesEditorViewport::Construct(const FArguments& InArgs, TSharedPtr<IRenderPageCollectionEditor> InBlueprintEditor)
+{
+	BlueprintEditorWeakPtr = InBlueprintEditor;
+	ViewportClient = MakeShareable(new FRenderPagesEditorViewportClient(nullptr, SharedThis(this)));
+	LevelSequencePlayerWorld = nullptr;
+	LevelSequencePlayerActor = nullptr;
+	LevelSequencePlayer = nullptr;
+	LevelSequence = nullptr;
+	Page = nullptr;
+	LevelSequenceTime = 0.0f;
+	bRenderedLastAttempt = false;
+
+	SEditorViewport::Construct(SEditorViewport::FArguments());
+}
+
+UE::RenderPages::Private::SRenderPagesEditorViewport::~SRenderPagesEditorViewport()
+{
+	DestroySequencePlayer();
+	ViewportClient.Reset();
+}
 
-	if (IsValid(Page) && IsValid(PageCollection))
+void UE::RenderPages::Private::SRenderPagesEditorViewport::Render()
+{
+	bRenderedLastAttempt = false;
+
+	if (!ViewportClient.IsValid())
 	{
-		PageCollection->PreRender(Page);
+		return;
 	}
 
-	bool bHasSetCameraData = false;
-	if (IsValid(SequencePlayer))
+	if (const TSharedPtr<IRenderPageCollectionEditor> BlueprintEditor = BlueprintEditorWeakPtr.Pin())
 	{
-		if (UCameraComponent* Camera = SequencePlayer->GetActiveCameraComponent(); IsValid(Camera))
+		if (BlueprintEditor->IsCurrentlyRenderingOrPlaying())
 		{
-			ViewportClient->SetViewLocation(Camera->GetComponentLocation());
-			ViewportClient->SetViewRotation(Camera->GetComponentRotation());
-			ViewportClient->ViewFOV = Camera->FieldOfView;
-			bHasSetCameraData = true;
+			return;
 		}
-	}
-	if (!bHasSetCameraData)
-	{
-		if (UWorld* World = GetWorld(); IsValid(World))
+
+		if (URenderPageCollection* PageCollection = BlueprintEditor->GetInstance(); IsValid(PageCollection))
 		{
-			if (APlayerController* LocalPlayerController = World->GetFirstPlayerController(); IsValid(LocalPlayerController))
+			if (!PageCollection->HasRenderPage(Page))
+			{
+				return;
+			}
+			bRenderedLastAttempt = true;
+
+			ULevelSequencePlayer* SequencePlayer = GetSequencePlayer();
+			if (IsValid(SequencePlayer))
 			{
-				FVector ViewLocation;
-				FRotator ViewRotation;
-				LocalPlayerController->GetPlayerViewPoint(ViewLocation, ViewRotation);
-				ViewportClient->SetViewLocation(ViewLocation);
-				ViewportClient->SetViewRotation(ViewRotation);
-				ViewportClient->ViewFOV = 90;
+				SequencePlayer->SetPlaybackPosition(FMovieSceneSequencePlaybackParams(LevelSequenceTime, EUpdatePositionMethod::Play));// execute this every tick, in case any sequencer values get overwritten (by remote control props for example)
 			}
-			else
+
+			if (IsValid(Page) && IsValid(PageCollection))
+			{
+				PageCollection->PreRender(Page);
+			}
+
+			bool bHasSetCameraData = false;
+			if (IsValid(SequencePlayer))
+			{
+				if (UCameraComponent* Camera = SequencePlayer->GetActiveCameraComponent(); IsValid(Camera))
+				{
+					ViewportClient->SetViewLocation(Camera->GetComponentLocation());
+					ViewportClient->SetViewRotation(Camera->GetComponentRotation());
+					ViewportClient->ViewFOV = Camera->FieldOfView;
+					bHasSetCameraData = true;
+				}
+			}
+			if (!bHasSetCameraData)
 			{
-				for (TActorIterator<APlayerStart> It(World); It; ++It)
+				if (UWorld* World = GetWorld(); IsValid(World))
 				{
-					if (APlayerStart* PlayerStart = *It; IsValid(PlayerStart))
+					if (APlayerController* LocalPlayerController = World->GetFirstPlayerController(); IsValid(LocalPlayerController))
 					{
-						ViewportClient->SetViewLocation(PlayerStart->GetActorLocation());
-						ViewportClient->SetViewRotation(PlayerStart->GetActorRotation());
+						FVector ViewLocation;
+						FRotator ViewRotation;
+						LocalPlayerController->GetPlayerViewPoint(ViewLocation, ViewRotation);
+						ViewportClient->SetViewLocation(ViewLocation);
+						ViewportClient->SetViewRotation(ViewRotation);
 						ViewportClient->ViewFOV = 90;
-						break;
+					}
+					else
+					{
+						for (TActorIterator<APlayerStart> It(World); It; ++It)
+						{
+							if (APlayerStart* PlayerStart = *It; IsValid(PlayerStart))
+							{
+								ViewportClient->SetViewLocation(PlayerStart->GetActorLocation());
+								ViewportClient->SetViewRotation(PlayerStart->GetActorRotation());
+								ViewportClient->ViewFOV = 90;
+								break;
+							}
+						}
 					}
 				}
 			}
-		}
-	}
 
-	ViewportClient->Viewport->Draw();
+			ViewportClient->Viewport->Draw();
 
-	if (IsValid(Page) && IsValid(PageCollection))
-	{
-		PageCollection->PostRender(Page);
+			if (IsValid(Page) && IsValid(PageCollection))
+			{
+				PageCollection->PostRender(Page);
+			}
+		}
 	}
 }
 
-void UE::RenderPages::Private::SRenderPagesEditorViewport::Construct(const FArguments& InArgs)
-{
-	ViewportClient = MakeShareable(new FRenderPagesEditorViewportClient(nullptr, SharedThis(this)));
-	LevelSequencePlayerWorld = nullptr;
-	LevelSequencePlayerActor = nullptr;
-	LevelSequencePlayer = nullptr;
-	LevelSequence = nullptr;
-	LevelSequenceTime = 0.0f;
-	Page = nullptr;
-	PageCollection = nullptr;
-
-	SEditorViewport::Construct(SEditorViewport::FArguments());
-}
-
-UE::RenderPages::Private::SRenderPagesEditorViewport::~SRenderPagesEditorViewport()
-{
-	DestroySequencePlayer();
-	ViewportClient.Reset();
-}
-
-bool UE::RenderPages::Private::SRenderPagesEditorViewport::ShowSequenceFrame(URenderPage* InPage, URenderPageCollection* InPageCollection, ULevelSequence* InSequence, const float InTime)
+bool UE::RenderPages::Private::SRenderPagesEditorViewport::ShowSequenceFrame(URenderPage* InPage, ULevelSequence* InSequence, const float InTime)
 {
 	Page = InPage;
-	PageCollection = InPageCollection;
 	LevelSequenceTime = InTime;
 	if (!IsValid(InSequence))
 	{
@@ -131,10 +155,6 @@ bool UE::RenderPages::Private::SRenderPagesEditorViewport::ShowSequenceFrame(URe
 		LevelSequence = InSequence;
 		DestroySequencePlayer();
 	}
-	if (ULevelSequencePlayer* SequencePlayer = GetSequencePlayer(); IsValid(SequencePlayer))
-	{
-		SequencePlayer->SetPlaybackPosition(FMovieSceneSequencePlaybackParams(LevelSequenceTime, EUpdatePositionMethod::Play));
-	}
 	return true;
 }
 
@@ -169,7 +189,6 @@ ULevelSequencePlayer* UE::RenderPages::Private::SRenderPagesEditorViewport::GetS
 			{
 				Player->Initialize(LevelSequence, World->PersistentLevel, PlaybackSettings, CameraSettings);
 				Player->State.AssignSequence(MovieSceneSequenceID::Root, *LevelSequence, *Player);
-				Player->SetPlaybackPosition(FMovieSceneSequencePlaybackParams(LevelSequence->GetMovieScene()->GetPlaybackRange().GetLowerBoundValue().Value, EUpdatePositionMethod::Play));
 
 				LevelSequencePlayerWorld = World;
 				LevelSequencePlayerActor = PlayerActor;
@@ -207,15 +226,16 @@ void UE::RenderPages::Private::SRenderPagesPageViewerLive::Construct(const FArgu
 {
 	BlueprintEditorWeakPtr = InBlueprintEditor;
 	SelectedPageWeakPtr = nullptr;
+	bViewportWidgetVisible = false;
 
-	SAssignNew(ViewportWidget, SRenderPagesEditorViewport)
-		.Visibility(EVisibility::Hidden);
+	SAssignNew(ViewportWidget, SRenderPagesEditorViewport, InBlueprintEditor);
 
 	SAssignNew(FrameSlider, SRenderPagesPageViewerFrameSlider)
 		.Visibility(EVisibility::Hidden)
 		.OnValueChanged(this, &SRenderPagesPageViewerLive::FrameSliderValueChanged);
 
 	SelectedPageChanged();
+	ViewportWidget->Render();// prevents the waiting text from showing up for 1 frame when switching from any other viewer mode to the live viewer mode
 
 	InBlueprintEditor->OnRenderPagesChanged().AddSP(this, &SRenderPagesPageViewerLive::PagesDataChanged);
 	InBlueprintEditor->OnRenderPagesSelectionChanged().AddSP(this, &SRenderPagesPageViewerLive::SelectedPageChanged);
@@ -225,24 +245,61 @@ void UE::RenderPages::Private::SRenderPagesPageViewerLive::Construct(const FArgu
 	[
 		SNew(SVerticalBox)
 
-		// viewport
+		// viewport & waiting text
 		+ SVerticalBox::Slot()
 		.FillHeight(1.0f)
 		.Padding(0.0f)
 		[
-			SNew(SScaleBox)
-			.HAlign(HAlign_Center)
-			.VAlign(VAlign_Center)
-			.Stretch(EStretch::ScaleToFit)
-			.StretchDirection(EStretchDirection::Both)
+			SNew(SOverlay)
+			.Visibility_Lambda([this]() -> EVisibility { return (bViewportWidgetVisible ? EVisibility::SelfHitTestInvisible : EVisibility::Hidden); })
+
+			// viewport
+			+ SOverlay::Slot()
+			.HAlign(HAlign_Fill)
+			.VAlign(VAlign_Fill)
+			[
+				SNew(SScaleBox)
+				.HAlign(HAlign_Center)
+				.VAlign(VAlign_Center)
+				.Stretch(EStretch::ScaleToFit)
+				.StretchDirection(EStretchDirection::Both)
+				[
+					SNew(SBox)
+					.HAlign(HAlign_Fill)
+					.VAlign(VAlign_Fill)
+					.WidthOverride(1.0f)
+					.HeightOverride_Lambda([this]() -> float { return 1.0 / (SelectedPageWeakPtr.IsValid() ? SelectedPageWeakPtr->GetOutputAspectRatio() : 1.0); })
+					[
+						ViewportWidget.ToSharedRef()
+					]
+				]
+			]
+
+			// waiting text
+			+ SOverlay::Slot()
+			.HAlign(HAlign_Fill)
+			.VAlign(VAlign_Fill)
 			[
-				SNew(SBox)
+				SNew(SOverlay)
+				.Visibility_Lambda([this]() -> EVisibility { return (ViewportWidget->HasRenderedLastAttempt() ? EVisibility::Hidden : EVisibility::HitTestInvisible); })
+
+				// image to hide the viewport
+				+ SOverlay::Slot()
 				.HAlign(HAlign_Fill)
 				.VAlign(VAlign_Fill)
-				.WidthOverride(1.0f)
-				.HeightOverride_Lambda([this]() -> float { return 1.0 / (SelectedPageWeakPtr.IsValid() ? SelectedPageWeakPtr->GetOutputAspectRatio() : 1.0); })
 				[
-					ViewportWidget.ToSharedRef()
+					SNew(SImage)
+					.Image(FAppStyle::GetBrush("Brushes.White"))
+					.ColorAndOpacity(FLinearColor(0.0185, 0.0185, 0.0185, 1))
+				]
+
+				// text
+				+ SOverlay::Slot()
+				.HAlign(HAlign_Center)
+				.VAlign(VAlign_Center)
+				[
+					SNew(STextBlock)
+					.Text(FText::FromString(TEXT("Waiting for renderer...")))
 				]
 			]
 		]
@@ -306,37 +363,32 @@ void UE::RenderPages::Private::SRenderPagesPageViewerLive::FrameSliderValueChang
 
 void UE::RenderPages::Private::SRenderPagesPageViewerLive::UpdateViewport()
 {
-	if (!ViewportWidget.IsValid())
+	bViewportWidgetVisible = false;
+
+	if (!ViewportWidget.IsValid() || !FrameSlider.IsValid())
 	{
 		return;
 	}
-	ViewportWidget->SetVisibility(EVisibility::Hidden);
 
-	if (FrameSlider.IsValid())
+	if (URenderPage* SelectedPage = SelectedPageWeakPtr.Get(); IsValid(SelectedPage))
 	{
-		if (URenderPage* SelectedPage = SelectedPageWeakPtr.Get(); IsValid(SelectedPage))
+		if (ULevelSequence* Sequence = SelectedPage->GetSequence(); IsValid(Sequence))
 		{
-			if (ULevelSequence* Sequence = SelectedPage->GetSequence(); IsValid(Sequence))
+			if (const TSharedPtr<IRenderPageCollectionEditor> BlueprintEditor = BlueprintEditorWeakPtr.Pin())
 			{
-				if (const TSharedPtr<IRenderPageCollectionEditor> BlueprintEditor = BlueprintEditorWeakPtr.Pin())
+				TOptional<double> StartTime = SelectedPage->GetStartTime();
+				TOptional<double> EndTime = SelectedPage->GetEndTime();
+				if (!StartTime.IsSet() || !EndTime.IsSet() || (StartTime.Get(0) > EndTime.Get(0)))
 				{
-					if (URenderPageCollection* PageCollection = BlueprintEditor->GetInstance(); IsValid(PageCollection))
-					{
-						TOptional<double> StartTime = SelectedPage->GetStartTime();
-						TOptional<double> EndTime = SelectedPage->GetEndTime();
-						if (!StartTime.IsSet() || !EndTime.IsSet() || (StartTime.Get(0) > EndTime.Get(0)))
-						{
-							return;
-						}
-
-						if (!ViewportWidget->ShowSequenceFrame(SelectedPage, PageCollection, Sequence, FMath::Lerp(StartTime.Get(0), EndTime.Get(0), FrameSlider->GetValue())))
-						{
-							return;
-						}
+					return;
+				}
 
-						ViewportWidget->SetVisibility(EVisibility::Visible);
-					}
+				if (!ViewportWidget->ShowSequenceFrame(SelectedPage, Sequence, FMath::Lerp(StartTime.Get(0), EndTime.Get(0), FrameSlider->GetValue())))
+				{
+					return;
 				}
+
+				bViewportWidgetVisible = true;
 			}
 		}
 	}
diff --git a/Engine/Plugins/Experimental/RenderPages/Source/RenderPagesEditor/Private/UI/SRenderPagesPageViewerLive.h b/Engine/Plugins/Experimental/RenderPages/Source/RenderPagesEditor/Private/UI/SRenderPagesPageViewerLive.h
--- a/Engine/Plugins/Experimental/RenderPages/Source/RenderPagesEditor/Private/UI/SRenderPagesPageViewerLive.h
+++ b/Engine/Plugins/Experimental/RenderPages/Source/RenderPagesEditor/Private/UI/SRenderPagesPageViewerLive.h
@@ -52,10 +52,12 @@ namespace UE::RenderPages::Private
 		SLATE_END_ARGS()
 
 		virtual void Tick(const FGeometry&, const double, const float) override;
-		void Construct(const FArguments& InArgs);
+		void Construct(const FArguments& InArgs, TSharedPtr<IRenderPageCollectionEditor> InBlueprintEditor);
 		virtual ~SRenderPagesEditorViewport() override;
 
-		bool ShowSequenceFrame(URenderPage* InPage, URenderPageCollection* InPageCollection, ULevelSequence* InSequence, const float InTime);
+		void Render();
+		bool ShowSequenceFrame(URenderPage* InPage, ULevelSequence* InSequence, const float InTime);
+		bool HasRenderedLastAttempt() const { return bRenderedLastAttempt; }
 
 	protected:
 		//~ Begin SEditorViewport Interface
@@ -68,6 +70,9 @@ namespace UE::RenderPages::Private
 		void DestroySequencePlayer();
 
 	private:
+		/** A reference to the BP Editor that owns this collection. */
+		TWeakPtr<IRenderPageCollectionEditor> BlueprintEditorWeakPtr;
+
 		/** The viewport client. */
 		TSharedPtr<FRenderPagesEditorViewportClient> ViewportClient;
 
@@ -87,16 +92,17 @@ namespace UE::RenderPages::Private
 		UPROPERTY()
 		TObjectPtr<ULevelSequence> LevelSequence;
 
-		/** The time of the currently playing sequence. */
-		float LevelSequenceTime;
-
 		/** The page that's currently being shown. */
 		UPROPERTY()
 		TObjectPtr<URenderPage> Page;
 
-		/** The collection of the page that's currently being shown. */
+		/** The time of the currently playing sequence. */
+		UPROPERTY()
+		float LevelSequenceTime;
+
+		/** Whether it rendered or not during the last tick. */
 		UPROPERTY()
-		TObjectPtr<URenderPageCollection> PageCollection;
+		bool bRenderedLastAttempt;
 	};
 
 
@@ -126,10 +132,14 @@ namespace UE::RenderPages::Private
 		/** A reference to the job that's currently rendering. */
 		TWeakObjectPtr<URenderPage> SelectedPageWeakPtr;
 
-		/** The scene viewport widget */
+		/** The viewport widget. */
 		TSharedPtr<SRenderPagesEditorViewport> ViewportWidget;
 
 		/** The widget that allows the user to select what frame they'd like to see. */
 		TSharedPtr<SRenderPagesPageViewerFrameSlider> FrameSlider;
+
+		/** Whether the viewport widget (and waiting text etc) should be visible or not. */
+		UPROPERTY()
+		bool bViewportWidgetVisible;
 	};
 }
diff --git a/Engine/Plugins/Experimental/RenderPages/Source/RenderPagesEditor/Private/UI/SRenderPagesPageViewerPreview.cpp b/Engine/Plugins/Experimental/RenderPages/Source/RenderPagesEditor/Private/UI/SRenderPagesPageViewerPreview.cpp
--- a/Engine/Plugins/Experimental/RenderPages/Source/RenderPagesEditor/Private/UI/SRenderPagesPageViewerPreview.cpp
+++ b/Engine/Plugins/Experimental/RenderPages/Source/RenderPagesEditor/Private/UI/SRenderPagesPageViewerPreview.cpp
@@ -87,26 +87,45 @@ void UE::RenderPages::Private::SRenderPagesPageViewerPreview::Construct(const FA
 			.IsFocusable(false)
 			.OnClicked(this, &SRenderPagesPageViewerPreview::OnClicked)
 			[
-				SNew(SScaleBox)
+				SNew(SOverlay)
+
+				// waiting text
+				+ SOverlay::Slot()
 				.HAlign(HAlign_Center)
 				.VAlign(VAlign_Center)
-				.Stretch(EStretch::ScaleToFit)
-				.StretchDirection(EStretchDirection::Both)
 				[
-					SNew(SOverlay)
-
-					+ SOverlay::Slot()
-					.HAlign(HAlign_Fill)
-					.VAlign(VAlign_Fill)
-					[
-						ImageBackground.ToSharedRef()
-					]
+					SNew(STextBlock)
+					.Text(FText::FromString(IsPreviewWidget() ? TEXT("Waiting for renderer...") : TEXT("Click here to start rendering")))
+				]
 
-					+ SOverlay::Slot()
-					.HAlign(HAlign_Fill)
-					.VAlign(VAlign_Fill)
+				// image & background
+				+ SOverlay::Slot()
+				.HAlign(HAlign_Fill)
+				.VAlign(VAlign_Fill)
+				[
+					SNew(SScaleBox)
+					.HAlign(HAlign_Center)
+					.VAlign(VAlign_Center)
+					.Stretch(EStretch::ScaleToFit)
+					.StretchDirection(EStretchDirection::Both)
 					[
-						Image.ToSharedRef()
+						SNew(SOverlay)
+
+						// background
+						+ SOverlay::Slot()
+						.HAlign(HAlign_Fill)
+						.VAlign(VAlign_Fill)
+						[
+							ImageBackground.ToSharedRef()
+						]
+
+						// image
+						+ SOverlay::Slot()
+						.HAlign(HAlign_Fill)
+						.VAlign(VAlign_Fill)
+						[
+							Image.ToSharedRef()
+						]
 					]
 				]
 			]
